[NNEDenoiser] Neural denoising platform for the Unreal Path Tracer based on the Neural Network Engine (NNE).

Problem:
There are efforts to port a high quality in-house path tracer denosier and Intel's Open Image Denoise to the Neural Network Engine (NNE). But a platform to run them is currently missing.

Solution:
Develop a plugin that allows licencees to either run the in-house denoiser, Intel's OIDN or a custom denoiser. Use the new path tracer denoise interface and NNE.

Tests:
* P4V Non-Unity build
* Editor, Win64:
  - Run OIDN models on CPU, GPU and where possible on RDG runtime and verify result in viewport visually
  - Export image sequence with Movie Render Queue and verify results visually

Note:
* In-house denoise model will be added later, since it is still in development

diff --git a/Engine/Plugins/Experimental/NNEDenoiser/NNEDenoiser.uplugin b/Engine/Plugins/Experimental/NNEDenoiser/NNEDenoiser.uplugin
new file mode 100644
--- /dev/null
+++ b/Engine/Plugins/Experimental/NNEDenoiser/NNEDenoiser.uplugin
@@ -0,0 +1,39 @@
+{
+	"FileVersion": 3,
+	"Version": 1,
+	"VersionName": "1.0",
+	"FriendlyName": "NNEDenoiser",
+	"Description": "Neural denoiser for the Unreal Path Tracer based on the Neural Network Engine (NNE).",
+	"Category": "Denoising",
+	"CreatedBy": "Epic Games, Inc.",
+	"CreatedByURL": "https://epicgames.com",
+	"DocsURL": "",
+	"MarketplaceURL": "",
+	"SupportURL": "",
+	"CanContainContent": true,
+	"IsBetaVersion": false,
+	"IsExperimentalVersion": true,
+	"Installed": false,
+	"Modules": [
+		{
+			"Name": "NNEDenoiser",
+			"Type": "Runtime",
+			"LoadingPhase": "Default"
+		},
+		{
+			"Name": "NNEDenoiserShaders",
+			"Type": "Runtime",
+			"LoadingPhase": "PostConfigInit"
+		}
+	],
+	"Plugins": [
+		{
+			"Name": "NNERuntimeORT",
+			"Enabled": true
+		},
+		{
+			"Name": "NNERuntimeRDG",
+			"Enabled": true
+		}
+	]
+}
\ No newline at end of file
diff --git a/Engine/Plugins/Experimental/NNEDenoiser/Source/NNEDenoiser/NNEDenoiser.Build.cs b/Engine/Plugins/Experimental/NNEDenoiser/Source/NNEDenoiser/NNEDenoiser.Build.cs
new file mode 100644
--- /dev/null
+++ b/Engine/Plugins/Experimental/NNEDenoiser/Source/NNEDenoiser/NNEDenoiser.Build.cs
@@ -0,0 +1,30 @@
+// Copyright Epic Games, Inc. All Rights Reserved.
+
+using UnrealBuildTool;
+
+public class NNEDenoiser : ModuleRules
+{
+	public NNEDenoiser(ReadOnlyTargetRules Target) : base(Target)
+	{
+		PCHUsage = PCHUsageMode.UseExplicitOrSharedPCHs;
+
+		PublicDependencyModuleNames.AddRange(
+			new string[]
+			{
+				"Core",
+				"RenderCore"
+			});
+		
+		PrivateDependencyModuleNames.AddRange(
+			new string[]
+			{
+				"CoreUObject",
+				"DeveloperSettings",
+				"Engine",
+				"NNE",
+				"NNEDenoiserShaders",
+				"Renderer",
+				"RHI"
+			});
+	}
+}
diff --git a/Engine/Plugins/Experimental/NNEDenoiser/Source/NNEDenoiser/Private/NNEDenoise.cpp b/Engine/Plugins/Experimental/NNEDenoiser/Source/NNEDenoiser/Private/NNEDenoise.cpp
new file mode 100644
--- /dev/null
+++ b/Engine/Plugins/Experimental/NNEDenoiser/Source/NNEDenoiser/Private/NNEDenoise.cpp
@@ -0,0 +1,30 @@
+// Copyright Epic Games, Inc. All Rights Reserved.
+
+#include "NNEDenoiser.h"
+#include "Misc/CoreDelegates.h"
+#include "NNEDenoiserLog.h"
+#include "NNEDenoiserViewExtension.h"
+
+DEFINE_LOG_CATEGORY(LogNNEDenoiser);
+
+void FNNEDenoiserModule::StartupModule()
+{
+	// This code will execute after your module is loaded into memory; the exact timing is specified in the .uplugin file per-module
+	UE_LOG(LogNNEDenoiser, Log, TEXT("NNEDenoiser module starting up"));
+
+	FCoreDelegates::OnPostEngineInit.AddLambda([this] ()
+	{
+		ViewExtension = FSceneViewExtensions::NewExtension<UE::NNEDenoiser::Private::FViewExtension>();
+	});
+}
+
+void FNNEDenoiserModule::ShutdownModule()
+{
+	// This function may be called during shutdown to clean up your module.  For modules that support dynamic reloading,
+	// we call this function before unloading the module.
+	UE_LOG(LogNNEDenoiser, Log, TEXT("NNEDenoiser module shut down"));
+
+	ViewExtension.Reset();
+}
+
+IMPLEMENT_MODULE(FNNEDenoiserModule, NNEDenoiser)
\ No newline at end of file
diff --git a/Engine/Plugins/Experimental/NNEDenoiser/Source/NNEDenoiser/Private/NNEDenoiser.h b/Engine/Plugins/Experimental/NNEDenoiser/Source/NNEDenoiser/Private/NNEDenoiser.h
new file mode 100644
--- /dev/null
+++ b/Engine/Plugins/Experimental/NNEDenoiser/Source/NNEDenoiser/Private/NNEDenoiser.h
@@ -0,0 +1,22 @@
+// Copyright Epic Games, Inc. All Rights Reserved.
+
+#pragma once
+
+#include "Modules/ModuleManager.h"
+#include "Templates/SharedPointer.h"
+
+namespace UE::NNEDenoiser::Private
+{
+	class FViewExtension;
+}
+
+class FNNEDenoiserModule : public IModuleInterface
+{
+public:
+	/** IModuleInterface implementation */
+	virtual void StartupModule() override;
+	virtual void ShutdownModule() override;
+
+private:
+	TSharedPtr<UE::NNEDenoiser::Private::FViewExtension, ESPMode::ThreadSafe> ViewExtension;
+};
diff --git a/Engine/Plugins/Experimental/NNEDenoiser/Source/NNEDenoiser/Private/NNEDenoiserGenericDenoiser.cpp b/Engine/Plugins/Experimental/NNEDenoiser/Source/NNEDenoiser/Private/NNEDenoiserGenericDenoiser.cpp
new file mode 100644
--- /dev/null
+++ b/Engine/Plugins/Experimental/NNEDenoiser/Source/NNEDenoiser/Private/NNEDenoiserGenericDenoiser.cpp
@@ -0,0 +1,195 @@
+// Copyright Epic Games, Inc. All Rights Reserved.
+
+#include "NNEDenoiserGenericDenoiser.h"
+#include "Algo/Transform.h"
+#include "NNEDenoiserHistory.h"
+#include "NNEDenoiserIOProcess.h"
+#include "NNEDenoiserLog.h"
+#include "NNEDenoiserModelInstance.h"
+#include "NNEDenoiserResourceManager.h"
+#include "NNEDenoiserTransferFunction.h"
+#include "NNEDenoiserUtils.h"
+#include <type_traits>
+
+namespace UE::NNEDenoiser::Private
+{
+
+class FResourceAccess : public IResourceAccess
+{
+public:
+	FResourceAccess(const FResourceManager& ResourceManager) : ResourceManager(ResourceManager) { }
+
+	virtual FRDGTextureRef GetTexture(EResourceName TensorName, int32 FrameIdx) const override
+	{
+		return ResourceManager.GetTexture(TensorName, FrameIdx);
+	}
+
+	virtual FRDGTextureRef GetIntermediateTexture(EResourceName TensorName, int32 FrameIdx) const override
+	{
+		return ResourceManager.GetIntermediateTexture(TensorName, FrameIdx);
+	}
+
+private:
+	const FResourceManager& ResourceManager;
+};
+
+FRDGBufferRef CreateBufferRDG(FRDGBuilder& GraphBuilder, const NNE::FTensorDesc& TensorDesc, const NNE::FTensorShape& TensorShape)
+{
+	const uint32 BytesPerElement = TensorDesc.GetElementByteSize();
+	const uint32 NumElements = TensorShape.Volume();
+
+	return GraphBuilder.CreateBuffer(FRDGBufferDesc::CreateBufferDesc(BytesPerElement, NumElements), *TensorDesc.GetName());
+}
+
+TArray<FRDGBufferRef> CreateBuffersRDG(FRDGBuilder& GraphBuilder, TConstArrayView<NNE::FTensorDesc> TensorDescs, TConstArrayView<NNE::FTensorShape> TensorShapes)
+{
+	check(TensorDescs.Num() == TensorShapes.Num());
+
+	TArray<FRDGBufferRef> Result;
+	Result.SetNumUninitialized(TensorDescs.Num());
+
+	for (int32 I = 0; I < TensorDescs.Num(); I++)
+	{
+		Result[I] = CreateBufferRDG(GraphBuilder, TensorDescs[I], TensorShapes[I]);
+	}
+	return Result;
+}
+
+TArray<NNE::FTensorBindingRDG> GetBindingRDG(TConstArrayView<FRDGBufferRef> Buffers)
+{
+	TArray<NNE::FTensorBindingRDG> Result;
+	Result.Reserve(Buffers.Num());
+	
+	Algo::Transform(Buffers, Result, [] (FRDGBuffer* Buffer) { return NNE::FTensorBindingRDG{Buffer}; });
+
+	return Result;
+}
+
+void RegisterTextureIfNeeded(FResourceManager& ResourceManager, FRDGTextureRef Texture, EResourceName ResourceName, const IInputProcess& InputProcess)
+{
+	const int32 MinNumFrames = ResourceName == EResourceName::Output ? 1 : 0;
+	const int32 NumFrames = FMath::Max(InputProcess.NumFrames(ResourceName), MinNumFrames);
+	if (NumFrames <= 0)
+	{
+		return;
+	}
+	
+	ResourceManager.AddTexture(ResourceName, Texture, NumFrames);
+}
+
+bool PrepareAndValidate(IModelInstance& ModelInstance, const IInputProcess& InputProcess, const IOutputProcess& OutputProcess, FIntPoint Extent)
+{
+	bool Result = InputProcess.PrepareAndValidate(ModelInstance, Extent);
+	Result &= OutputProcess.Validate(ModelInstance, Extent);
+
+	return Result;
+}
+
+void AddTilePasses(FRDGBuilder& GraphBuilder, IModelInstance& ModelInstance, const IInputProcess& InputProcess, const IOutputProcess& OutputProcess,
+		FResourceManager& ResourceManager, TConstArrayView<FRDGBufferRef> InputBuffers, TConstArrayView<FRDGBufferRef> OutputBuffers)
+{
+	FResourceAccess ResourceAccess(ResourceManager);
+
+	// 1. Read input and write to input buffers
+	InputProcess.AddPasses(GraphBuilder, ModelInstance.GetInputTensorDescs(), ModelInstance.GetInputTensorShapes(), ResourceAccess, InputBuffers);
+
+	// 2. Create buffer binding and infer the model
+	bool bSuccess = !ModelInstance.EnqueueRDG(GraphBuilder, GetBindingRDG(InputBuffers), GetBindingRDG(OutputBuffers));
+	check(bSuccess);
+
+	// 3. Write output based on output buffer
+	FRDGTextureRef OutputTexture = ResourceManager.GetTexture(EResourceName::Output, 0);
+	OutputProcess.AddPasses(GraphBuilder, ModelInstance.GetOutputTensorDescs(), ModelInstance.GetOutputTensorShapes(), ResourceAccess, OutputBuffers, OutputTexture);
+}
+
+FGenericDenoiser::FGenericDenoiser(TUniquePtr<IModelInstance> ModelInstance, TUniquePtr<IInputProcess> InputProcess, TUniquePtr<IOutputProcess> OutputProcess, FParameters DenoiserParameters) :
+	ModelInstance(MoveTemp(ModelInstance)), InputProcess(MoveTemp(InputProcess)), OutputProcess(MoveTemp(OutputProcess)), DenoiserParameters(DenoiserParameters)
+{
+
+}
+
+FGenericDenoiser::~FGenericDenoiser()
+{
+
+}
+
+TUniquePtr<FHistory> FGenericDenoiser::AddPasses(
+	FRDGBuilder& GraphBuilder,
+	FRDGTextureRef ColorTex,
+	FRDGTextureRef AlbedoTex,
+	FRDGTextureRef NormalTex,
+	FRDGTextureRef OutputTex,
+	FRDGTextureRef FlowTex,
+	const FRHIGPUMask& GPUMask,
+	FHistory* History)
+{
+	using namespace UE::NNEDenoiserShaders::Internal;
+
+	SCOPED_NAMED_EVENT_TEXT("NNEDenoiser.AddPasses", FColor::Magenta);
+
+	// note: FlowTex is currently optional and can be null
+	check(ColorTex != FRDGTextureRef{});
+	check(AlbedoTex != FRDGTextureRef{});
+	check(NormalTex != FRDGTextureRef{});
+	check(OutputTex != FRDGTextureRef{});
+
+	const FIntPoint Extent = ColorTex->Desc.Extent;
+
+	TUniquePtr<FHistory> Result;
+
+	if (Extent == LastExtent || PrepareAndValidate(*ModelInstance, *InputProcess, *OutputProcess, Extent))
+	{
+		LastExtent = Extent;
+
+		const TConstArrayView<uint32> InputShapeData = ModelInstance->GetInputTensorShapes()[0].GetData();
+		const FIntPoint ModelInputSize{(int32)InputShapeData[3], (int32)InputShapeData[2]};
+
+		const FTiling Tiling = CreateTiling(ModelInputSize, DenoiserParameters.TileMinimumOverlap, Extent);
+
+		TMap<EResourceName, TArray<TRefCountPtr<IPooledRenderTarget>>> ResourceMap;
+		if (History)
+		{
+			ResourceMap = History->GetResourceMap();
+		}
+
+		FResourceManager ResourceManager(GraphBuilder, Tiling, ResourceMap);
+		RegisterTextureIfNeeded(ResourceManager, ColorTex, EResourceName::Color, *InputProcess);
+		RegisterTextureIfNeeded(ResourceManager, AlbedoTex, EResourceName::Albedo, *InputProcess);
+		RegisterTextureIfNeeded(ResourceManager, NormalTex, EResourceName::Normal, *InputProcess);
+		if (FlowTex != FRDGTextureRef{})
+		{
+			RegisterTextureIfNeeded(ResourceManager, FlowTex, EResourceName::Flow, *InputProcess);
+		}
+		RegisterTextureIfNeeded(ResourceManager, OutputTex, EResourceName::Output, *InputProcess);
+
+		TArray<FRDGBufferRef> InputBuffers = CreateBuffersRDG(GraphBuilder, ModelInstance->GetInputTensorDescs(), ModelInstance->GetInputTensorShapes());
+		TArray<FRDGBufferRef> OutputBuffers = CreateBuffersRDG(GraphBuilder, ModelInstance->GetOutputTensorDescs(), ModelInstance->GetOutputTensorShapes());
+
+		UE_LOG(LogNNEDenoiser, Log, TEXT("Divided work into %dx%d tiles of size %dx%d each..."), Tiling.Count.X, Tiling.Count.Y, Tiling.TileSize.X, Tiling.TileSize.Y);
+
+		for (int32 I = 0; I < Tiling.Tiles.Num(); I++)
+		{
+			ResourceManager.BeginTile(I);
+
+			// Do inference on tile
+			AddTilePasses(GraphBuilder, *ModelInstance, *InputProcess, *OutputProcess, ResourceManager, InputBuffers, OutputBuffers);
+
+			ResourceManager.EndTile();
+		}
+
+		ResourceMap = ResourceManager.MakeHistoryResourceMap();
+		if (!ResourceMap.IsEmpty())
+		{
+			Result = MakeUnique<FHistory>(*DebugName, MoveTemp(ResourceMap));
+		}
+	}
+	else
+	{
+		// If something went wrong, just copy color to output...
+		AddCopyTexturePass(GraphBuilder, ColorTex, OutputTex, FRHICopyTextureInfo{});
+	}
+
+	return Result;
+}
+
+} // namespace UE::NNEDenoiser::Private
\ No newline at end of file
diff --git a/Engine/Plugins/Experimental/NNEDenoiser/Source/NNEDenoiser/Private/NNEDenoiserGenericDenoiser.h b/Engine/Plugins/Experimental/NNEDenoiser/Source/NNEDenoiser/Private/NNEDenoiserGenericDenoiser.h
new file mode 100644
--- /dev/null
+++ b/Engine/Plugins/Experimental/NNEDenoiser/Source/NNEDenoiser/Private/NNEDenoiserGenericDenoiser.h
@@ -0,0 +1,52 @@
+// Copyright Epic Games, Inc. All Rights Reserved.
+
+#pragma once
+
+#include "Containers/UnrealString.h"
+#include "Math/IntPoint.h"
+#include "NNEDenoiserParameters.h"
+#include "RenderGraphFwd.h"
+#include "Templates/UniquePtr.h"
+
+struct FRHIGPUMask;
+
+namespace UE::NNEDenoiser::Private
+{
+
+class FHistory;
+class IModelInstance;
+class IInputProcess;
+class IOutputProcess;
+
+class FGenericDenoiser
+{
+public:
+	FGenericDenoiser(TUniquePtr<IModelInstance> ModelInstance, TUniquePtr<IInputProcess> InputProcess, TUniquePtr<IOutputProcess> OutputProcess, FParameters DenoiserParameters);
+
+	~FGenericDenoiser();
+
+	// IPathTracingSpatialTemporalDenoiser interface
+	const TCHAR* GetDebugName() const { return *DebugName; }
+
+	TUniquePtr<FHistory> AddPasses(
+		FRDGBuilder &GraphBuilder,
+		FRDGTextureRef ColorTex,
+		FRDGTextureRef AlbedoTex,
+		FRDGTextureRef NormalTex,
+		FRDGTextureRef OutputTex,
+		FRDGTextureRef FlowTex,
+		const FRHIGPUMask& GPUMask,
+		FHistory* History);
+
+private:
+	inline static const FString DebugName = TEXT("FGenericDenoiser");
+
+	TUniquePtr<IModelInstance> ModelInstance;
+	TUniquePtr<IInputProcess> InputProcess;
+	TUniquePtr<IOutputProcess> OutputProcess;
+	FParameters DenoiserParameters;
+
+	FIntPoint LastExtent = { -1, -1 };
+};
+
+} // namespace UE::NNEDenoiser::Private
\ No newline at end of file
diff --git a/Engine/Plugins/Experimental/NNEDenoiser/Source/NNEDenoiser/Private/NNEDenoiserHistory.h b/Engine/Plugins/Experimental/NNEDenoiser/Source/NNEDenoiser/Private/NNEDenoiserHistory.h
new file mode 100644
--- /dev/null
+++ b/Engine/Plugins/Experimental/NNEDenoiser/Source/NNEDenoiser/Private/NNEDenoiserHistory.h
@@ -0,0 +1,36 @@
+// Copyright Epic Games, Inc. All Rights Reserved.
+
+#pragma once
+
+#include "CoreMinimal.h"
+#include "NNEDenoiserLog.h"
+#include "NNEDenoiserModelIOMappingData.h"
+#include "PathTracingDenoiser.h"
+#include "RendererInterface.h"
+
+struct IPooledRenderTarget;
+
+namespace UE::NNEDenoiser::Private
+{
+
+using UE::Renderer::Private::IPathTracingSpatialTemporalDenoiser;
+
+class FHistory : public IPathTracingSpatialTemporalDenoiser::IHistory
+{
+public:
+	FHistory(const TCHAR* DebugName, TMap<EResourceName, TArray<TRefCountPtr<IPooledRenderTarget>>> ResourceMap) : DebugName(DebugName), ResourceMap(MoveTemp(ResourceMap)) {
+
+	}
+
+	virtual ~FHistory() = default;
+
+	const TCHAR* GetDebugName() const override { return DebugName; };
+
+	const TMap<EResourceName, TArray<TRefCountPtr<IPooledRenderTarget>>>& GetResourceMap() const { return ResourceMap; }
+
+private:
+	const TCHAR* DebugName;
+	TMap<EResourceName, TArray<TRefCountPtr<IPooledRenderTarget>>> ResourceMap;
+};
+
+}
\ No newline at end of file
diff --git a/Engine/Plugins/Experimental/NNEDenoiser/Source/NNEDenoiser/Private/NNEDenoiserIOProcess.h b/Engine/Plugins/Experimental/NNEDenoiser/Source/NNEDenoiser/Private/NNEDenoiserIOProcess.h
new file mode 100644
--- /dev/null
+++ b/Engine/Plugins/Experimental/NNEDenoiser/Source/NNEDenoiser/Private/NNEDenoiserIOProcess.h
@@ -0,0 +1,55 @@
+// Copyright Epic Games, Inc. All Rights Reserved.
+
+#pragma once
+
+#include "NNEDenoiserModelIOMappingData.h"
+#include "NNETypes.h"
+#include "RenderGraphFwd.h"
+
+namespace UE::NNEDenoiser::Private
+{
+
+class FResourceMapping;
+class IModelInstance;
+
+class IResourceAccess
+{
+public:
+	virtual FRDGTextureRef GetTexture(EResourceName TensorName, int32 FrameIdx) const = 0;
+	virtual FRDGTextureRef GetIntermediateTexture(EResourceName TensorName, int32 FrameIdx) const = 0;
+};
+
+class IInputProcess
+{
+public:
+	virtual ~IInputProcess() = default;
+
+	virtual bool PrepareAndValidate(IModelInstance& ModelInstance, FIntPoint Extent) const = 0;
+
+	virtual int32 NumFrames(EResourceName Name) const = 0;
+
+	virtual void AddPasses(
+		FRDGBuilder& GraphBuilder,
+		TConstArrayView<UE::NNE::FTensorDesc> TensorDescs,
+		TConstArrayView<UE::NNE::FTensorShape> TensorShapes,
+		const IResourceAccess& ResourceAccess,
+		TConstArrayView<FRDGBufferRef> OutputBuffers) const = 0;
+};
+
+class IOutputProcess
+{
+public:
+	virtual ~IOutputProcess() = default;
+
+	virtual bool Validate(const IModelInstance& ModelInstance, FIntPoint Extent) const = 0;
+
+	virtual void AddPasses(
+		FRDGBuilder& GraphBuilder,
+		TConstArrayView<UE::NNE::FTensorDesc> TensorDescs,
+		TConstArrayView<UE::NNE::FTensorShape> TensorShapes,
+		const IResourceAccess& ResourceAccess,
+		TConstArrayView<FRDGBufferRef> Buffers,
+		FRDGTextureRef OutputTexture) const = 0;
+};
+
+} // namespace UE::NNEDenoiser::Private
\ No newline at end of file
diff --git a/Engine/Plugins/Experimental/NNEDenoiser/Source/NNEDenoiser/Private/NNEDenoiserIOProcessBase.cpp b/Engine/Plugins/Experimental/NNEDenoiser/Source/NNEDenoiser/Private/NNEDenoiserIOProcessBase.cpp
new file mode 100644
--- /dev/null
+++ b/Engine/Plugins/Experimental/NNEDenoiser/Source/NNEDenoiser/Private/NNEDenoiserIOProcessBase.cpp
@@ -0,0 +1,358 @@
+// Copyright Epic Games, Inc. All Rights Reserved.
+
+#include "NNEDenoiserIOProcessBase.h"
+#include "NNEDenoiserLog.h"
+#include "NNEDenoiserModelInstance.h"
+#include "NNEDenoiserParameters.h"
+#include "NNEDenoiserShadersDefaultCS.h"
+#include "NNEDenoiserUtils.h"
+#include "RHIStaticStates.h"
+
+DECLARE_GPU_STAT_NAMED(FNNEDenoiserReadInput, TEXT("NNEDenoiser.ReadInput"));
+DECLARE_GPU_STAT_NAMED(FNNEDenoiserWriteOutput, TEXT("NNEDenoiser.WriteOutput"));
+
+namespace UE::NNEDenoiser::Private
+{
+
+namespace IOProcessBaseHelper
+{
+
+void AddReadInputPass(
+	FRDGBuilder& GraphBuilder,
+	FRDGTextureRef InputTexture,
+	FRDGBufferUAVRef InputBufferUAV,
+	FIntPoint InputBufferSize,
+	UE::NNEDenoiserShaders::Internal::ENNEDenoiserDataType DataType,
+	const TArray<FIntPoint>& ChannelMapping)
+{
+	using namespace UE::NNEDenoiserShaders::Internal;
+
+	if (ChannelMapping.IsEmpty())
+	{
+		UE_LOG(LogNNEDenoiser, Warning, TEXT("AddReadInputPass: ChannelMapping is empty. Nothing to do!"));
+		return;
+	}
+	
+	if (ChannelMapping.Num() >= FNNEDenoiserConstants::MAX_NUM_MAPPED_CHANNELS)
+	{
+		UE_LOG(LogNNEDenoiser, Warning, TEXT("AddReadInputPass: ChannelMapping has too many entries!"));
+		return;
+	}
+
+	const FIntVector InputTextureSize = InputTexture->Desc.GetSize();
+
+	FNNEDenoiserReadInputCS::FParameters *ReadInputParameters = GraphBuilder.AllocParameters<FNNEDenoiserReadInputCS::FParameters>();
+	ReadInputParameters->InputTextureWidth = InputTextureSize.X;
+	ReadInputParameters->InputTextureHeight = InputTextureSize.Y;
+	ReadInputParameters->InputTexture = InputTexture;
+	ReadInputParameters->InputBufferWidth = InputBufferSize.X;
+	ReadInputParameters->InputBufferHeight = InputBufferSize.Y;
+	ReadInputParameters->InputBuffer = InputBufferUAV;
+	for (int32 Idx = 0; Idx < ChannelMapping.Num(); Idx++)
+	{
+		ReadInputParameters->BufferChannel_TextureChannel_Unused_Unused[Idx] = { ChannelMapping[Idx].X, ChannelMapping[Idx].Y, 0, 0 };
+	}
+
+	FNNEDenoiserReadInputCS::FPermutationDomain PermutationVector;
+	PermutationVector.Set<FNNEDenoiserReadInputCS::FNNEDenoiserDataType>(DataType);
+	PermutationVector.Set<FNNEDenoiserReadInputCS::FNNEDenoiserNumMappedChannels>(ChannelMapping.Num());
+
+	FIntVector ReadInputThreadGroupCount = FIntVector(
+		FMath::DivideAndRoundUp(InputBufferSize.X, FNNEDenoiserConstants::THREAD_GROUP_SIZE),
+		FMath::DivideAndRoundUp(InputBufferSize.Y, FNNEDenoiserConstants::THREAD_GROUP_SIZE),
+		1);
+
+	FGlobalShaderMap* GlobalShaderMap = GetGlobalShaderMap(GMaxRHIFeatureLevel);
+	TShaderMapRef<FNNEDenoiserReadInputCS> ReadInputShader(GlobalShaderMap, PermutationVector);
+
+	RDG_EVENT_SCOPE(GraphBuilder, "NNEDenoiser.ReadInput");
+	RDG_GPU_STAT_SCOPE(GraphBuilder, FNNEDenoiserReadInput);
+
+	FComputeShaderUtils::AddPass(
+		GraphBuilder,
+		RDG_EVENT_NAME("NNEDenoiser.ReadInput"),
+		ERDGPassFlags::Compute | ERDGPassFlags::NeverCull,
+		ReadInputShader,
+		ReadInputParameters,
+		ReadInputThreadGroupCount);
+}
+
+void AddWriteOutputPass(
+	FRDGBuilder& GraphBuilder,
+	FRDGBufferUAVRef OutputBufferUAV,
+	FIntPoint OutputBufferSize,
+	FRDGTextureUAVRef OutputTextureUAV,
+	UE::NNEDenoiserShaders::Internal::ENNEDenoiserDataType DataType,
+	const TArray<FIntPoint>& ChannelMapping)
+{
+	using namespace UE::NNEDenoiserShaders::Internal;
+
+	const FIntVector OutputTextureSize = OutputTextureUAV->Desc.Texture->Desc.GetSize();
+
+	FNNEDenoiserWriteOutputCS::FParameters *WriteOutputParameters = GraphBuilder.AllocParameters<FNNEDenoiserWriteOutputCS::FParameters>();
+	WriteOutputParameters->OutputBufferWidth = OutputBufferSize.X;
+	WriteOutputParameters->OutputBufferHeight = OutputBufferSize.Y;
+	WriteOutputParameters->OutputBuffer = OutputBufferUAV;
+	WriteOutputParameters->OutputTextureWidth = OutputTextureSize.X;
+	WriteOutputParameters->OutputTextureHeight = OutputTextureSize.Y;
+	WriteOutputParameters->OutputTexture = OutputTextureUAV;
+	for (int32 Idx = 0; Idx < ChannelMapping.Num(); Idx++)
+	{
+		WriteOutputParameters->BufferChannel_TextureChannel_Unused_Unused[Idx] = { ChannelMapping[Idx].X, ChannelMapping[Idx].Y, 0, 0 };
+	}
+
+	FNNEDenoiserWriteOutputCS::FPermutationDomain PermutationVector;
+	PermutationVector.Set<FNNEDenoiserWriteOutputCS::FNNEDenoiserDataType>(DataType);
+	PermutationVector.Set<FNNEDenoiserWriteOutputCS::FNNEDenoiserNumMappedChannels>(ChannelMapping.Num());
+
+	FIntVector WriteOutputThreadGroupCount = FIntVector(
+		FMath::DivideAndRoundUp(OutputTextureSize.X, FNNEDenoiserConstants::THREAD_GROUP_SIZE),
+		FMath::DivideAndRoundUp(OutputTextureSize.Y, FNNEDenoiserConstants::THREAD_GROUP_SIZE),
+		1);
+
+	FGlobalShaderMap* GlobalShaderMap = GetGlobalShaderMap(GMaxRHIFeatureLevel);
+	TShaderMapRef<FNNEDenoiserWriteOutputCS> WriteOutputShader(GlobalShaderMap, PermutationVector);
+
+	RDG_EVENT_SCOPE(GraphBuilder, "NNEDenoiser.WriteOutput");
+	RDG_GPU_STAT_SCOPE(GraphBuilder, FNNEDenoiserWriteOutput);
+
+	FComputeShaderUtils::AddPass(
+		GraphBuilder,
+		RDG_EVENT_NAME("NNEDenoiser.WriteOutput"),
+		ERDGPassFlags::Compute | ERDGPassFlags::NeverCull,
+		WriteOutputShader,
+		WriteOutputParameters,
+		WriteOutputThreadGroupCount);
+}
+
+void AddReadInputPassForKind(
+	FRDGBuilder& GraphBuilder,
+	const IResourceAccess& ResourceAccess,
+	FIntPoint BufferSize,
+	NNEDenoiserShaders::Internal::ENNEDenoiserDataType DataType,
+	EResourceName TensorName,
+	const FResourceMapping& ResourceMapping,
+	FRDGBufferUAVRef BufferUAV)
+{
+	for (const TPair<int32, TArray<FIntPoint>>& ChannelMappingKeyValue : ResourceMapping.GetChannelMappingPerFrame(TensorName))
+	{
+		const int32 FrameIdx = -ChannelMappingKeyValue.Key;
+		FRDGTextureRef InputTexture = ResourceAccess.GetTexture(TensorName, FrameIdx);
+
+		AddReadInputPass(GraphBuilder, InputTexture, BufferUAV, BufferSize, DataType, ChannelMappingKeyValue.Value);
+	}
+};
+
+} // namespace IOProcessBaseHelper
+
+bool FInputProcessBase::PrepareAndValidate(IModelInstance& ModelInstance, FIntPoint Extent) const
+{
+	SCOPED_NAMED_EVENT_TEXT("NNEDenoiser.Prepare", FColor::Magenta);
+
+	const int32 NumBatches = 1;
+
+	UE_LOG(LogNNEDenoiser, Log, TEXT("Configure model for extent %dx%d..."), Extent.X, Extent.Y);
+
+	TConstArrayView<UE::NNE::FTensorDesc> InputTensorDescs = ModelInstance.GetInputTensorDescs();
+	if (InputTensorDescs.Num() != InputLayout.Num())
+	{
+		UE_LOG(LogNNEDenoiser, Error, TEXT("Wrong number of inputs (expected %d, got %d)!"), InputLayout.Num(), InputTensorDescs.Num())
+		return false;
+	}
+
+	UE_LOG(LogNNEDenoiser, Log, TEXT("Input shapes (set):"));
+
+	TArray<UE::NNE::FTensorShape> InputShapes;
+	for (int32 Idx = 0; Idx < InputTensorDescs.Num(); Idx++)
+	{
+		TConstArrayView<int32> InputSymbolicTensorShapeData = InputTensorDescs[Idx].GetShape().GetData();
+		const TArray<int32, TInlineAllocator<4>> RequiredInputShapeData = { NumBatches, InputLayout.NumChannels(Idx), -1, -1 };
+
+		if (!IsTensorShapeValid(InputSymbolicTensorShapeData, RequiredInputShapeData, TEXT("Input")))
+		{
+			return false;
+		}
+
+		const int32 ModelInputWidth = InputSymbolicTensorShapeData[3] >= 0 ? InputSymbolicTensorShapeData[3] : Extent.X;
+		const int32 ModelInputHeight = InputSymbolicTensorShapeData[2] >= 0 ? InputSymbolicTensorShapeData[2] : Extent.Y;
+
+		if (Extent.X < ModelInputWidth || Extent.Y < ModelInputHeight)
+		{
+			UE_LOG(LogNNEDenoiser, Error, TEXT("Input image too small (model expects at least size %dx%d, got %dx%d!"), ModelInputWidth, ModelInputHeight, Extent.X, Extent.Y);
+			return false;
+		}
+
+		const FIntVector4 ModelInputShape = { 1, InputSymbolicTensorShapeData[1], ModelInputHeight, ModelInputWidth };
+
+		InputShapes.Add(UE::NNE::FTensorShape::Make({ (uint32)ModelInputShape.X, (uint32)ModelInputShape.Y, (uint32)ModelInputShape.Z, (uint32)ModelInputShape.W }));
+
+		UE_LOG(LogNNEDenoiser, Log, TEXT("%d: (%d, %d, %d, %d)"), Idx, ModelInputShape.X, ModelInputShape.Y, ModelInputShape.Z, ModelInputShape.W)
+	}
+
+	const int32 Success = ModelInstance.SetInputTensorShapes(InputShapes);
+	if (Success != 0)
+	{
+		UE_LOG(LogNNEDenoiser, Error, TEXT("Could not configure model instance (ModelInstance.SetInputTensorShapes() failed)!"))
+		return false;
+	}
+
+	return true;
+}
+
+int32 FInputProcessBase::NumFrames(EResourceName Name) const
+{
+	return InputLayout.NumFrames(Name);
+}
+
+void FInputProcessBase::AddPasses(
+		FRDGBuilder& GraphBuilder,
+		TConstArrayView<UE::NNE::FTensorDesc> TensorDescs,
+		TConstArrayView<UE::NNE::FTensorShape> TensorShapes,
+		const IResourceAccess& ResourceAccess,
+		TConstArrayView<FRDGBufferRef> OutputBuffers) const
+{
+	const UEnum* ResourceNameEnum = StaticEnum<EResourceName>();
+	for (int32 I = 0; I < ResourceNameEnum->NumEnums(); I++)
+	{
+		const EResourceName TensorName = (EResourceName)ResourceNameEnum->GetValueByIndex(I);
+		const int32 NumFrames = InputLayout.NumFrames(TensorName);
+
+		for (int32 J = 0; J < NumFrames; J++)
+		{
+			if (!HasPreprocessInput(TensorName, J))
+			{
+				continue;
+			}
+
+			FRDGTextureRef InputTexture = ResourceAccess.GetTexture(TensorName, J);
+			FRDGTextureRef PreprocessedInputTexture = ResourceAccess.GetIntermediateTexture(TensorName, J);
+
+			PreprocessInput(GraphBuilder, InputTexture, TensorName, J, PreprocessedInputTexture);
+		}
+	}
+
+	for (int32 I = 0; I < TensorDescs.Num(); I++)
+	{
+		WriteInputBuffer(GraphBuilder, TensorDescs[I], TensorShapes[I], ResourceAccess, InputLayout.GetChecked(I), OutputBuffers[I]);
+	}
+}
+
+void FInputProcessBase::WriteInputBuffer(
+	FRDGBuilder& GraphBuilder,
+	const UE::NNE::FTensorDesc& TensorDesc,
+	const UE::NNE::FTensorShape& TensorShape,
+	const IResourceAccess& ResourceAccess,
+	const FResourceMapping& ResourceMapping,
+	FRDGBufferRef Buffer) const
+{
+	const ENNETensorDataType TensorDataType = TensorDesc.GetDataType();
+	const NNEDenoiserShaders::Internal::ENNEDenoiserDataType DataType = GetDenoiserShaderDataType(TensorDataType);
+	const FIntPoint BufferSize(TensorShape.GetData()[3], TensorShape.GetData()[2]);
+
+	FRDGBufferUAVRef BufferUAV = GraphBuilder.CreateUAV(FRDGBufferUAVDesc(Buffer, GetBufferFormat(TensorDataType)));
+
+	IOProcessBaseHelper::AddReadInputPassForKind(GraphBuilder, ResourceAccess, BufferSize, DataType, EResourceName::Color, ResourceMapping, BufferUAV);
+	IOProcessBaseHelper::AddReadInputPassForKind(GraphBuilder, ResourceAccess, BufferSize, DataType, EResourceName::Albedo, ResourceMapping, BufferUAV);
+	IOProcessBaseHelper::AddReadInputPassForKind(GraphBuilder, ResourceAccess, BufferSize, DataType, EResourceName::Normal, ResourceMapping, BufferUAV);
+	IOProcessBaseHelper::AddReadInputPassForKind(GraphBuilder, ResourceAccess, BufferSize, DataType, EResourceName::Flow, ResourceMapping, BufferUAV);
+	IOProcessBaseHelper::AddReadInputPassForKind(GraphBuilder, ResourceAccess, BufferSize, DataType, EResourceName::Output, ResourceMapping, BufferUAV);
+}
+
+bool FOutputProcessBase::Validate(const IModelInstance& ModelInstance, FIntPoint Extent) const
+{
+	SCOPED_NAMED_EVENT_TEXT("NNEDenoiser.Prepare", FColor::Magenta);
+
+	const int32 NumBatches = 1;
+
+	TConstArrayView<UE::NNE::FTensorDesc> OutputTensorDescs = ModelInstance.GetOutputTensorDescs();
+	if (OutputTensorDescs.Num() != OutputLayout.Num())
+	{
+		UE_LOG(LogNNEDenoiser, Error, TEXT("Wrong number of outputs (expected %d, got %d)!"), OutputLayout.Num(), OutputTensorDescs.Num())
+		return true;
+	}
+
+	for (int32 Idx = 0; Idx < OutputTensorDescs.Num(); Idx++)
+	{
+		TConstArrayView<int32> OutputSymbolicShapeData = OutputTensorDescs[Idx].GetShape().GetData();
+		const TArray<int32, TInlineAllocator<4>> RequiredOutputShapeData = { NumBatches, OutputLayout.NumChannels(Idx), -1, -1 };
+
+		if (!IsTensorShapeValid(OutputSymbolicShapeData, RequiredOutputShapeData, TEXT("Output")))
+		{
+			return false;
+		}
+	}
+
+	TConstArrayView<UE::NNE::FTensorShape> OutputShapes = ModelInstance.GetOutputTensorShapes();
+	if (OutputShapes.Num() != OutputLayout.Num())
+	{
+		UE_LOG(LogNNEDenoiser, Error, TEXT("Wrong number of output shapes or not resolved yet!"));
+		return true;
+	}
+	
+	UE_LOG(LogNNEDenoiser, Log, TEXT("Output shapes (resolved):"));
+
+	for (int32 Idx = 0; Idx < OutputTensorDescs.Num(); Idx++)
+	{
+		TConstArrayView<uint32> OutputShapeData = OutputShapes[Idx].GetData();
+		const TArray<int32, TInlineAllocator<4>> RequiredOutputShapeData = { NumBatches, OutputLayout.NumChannels(Idx), -1, -1 };
+
+		if (!IsTensorShapeValid(OutputShapeData, RequiredOutputShapeData, TEXT("Output")))
+		{
+			return false;
+		}
+	
+		UE_LOG(LogNNEDenoiser, Log, TEXT("%d: (%d, %d, %d, %d)"), Idx, OutputShapeData[0], OutputShapeData[1], OutputShapeData[2], OutputShapeData[3]);
+	}
+
+	return true;
+}
+
+void FOutputProcessBase::AddPasses(
+	FRDGBuilder& GraphBuilder,
+	TConstArrayView<UE::NNE::FTensorDesc> TensorDescs,
+	TConstArrayView<UE::NNE::FTensorShape> TensorShapes,
+	const IResourceAccess& ResourceAccess,
+	TConstArrayView<FRDGBufferRef> Buffers,
+	FRDGTextureRef OutputTexture) const
+{
+	for (int32 I = 0; I < TensorDescs.Num(); I++)
+	{
+		ReadOutputBuffer(GraphBuilder, TensorDescs[I], TensorShapes[I], ResourceAccess, Buffers[I], OutputLayout.GetChecked(I), OutputTexture);
+	}
+
+	const EResourceName TensorName = EResourceName::Output;
+	const int32 FrameIdx = 0;
+
+	if (HasPostprocessOutput(TensorName, FrameIdx))
+	{
+		FRDGTextureRef PostprocessInputTexture = ResourceAccess.GetTexture(EResourceName::Output, FrameIdx);
+		FRDGTextureRef PostprocessOutputTexture = ResourceAccess.GetIntermediateTexture(EResourceName::Output, FrameIdx);
+		
+		PostprocessOutput(GraphBuilder, PostprocessInputTexture, PostprocessOutputTexture);
+	}
+}
+
+void FOutputProcessBase::ReadOutputBuffer(
+	FRDGBuilder& GraphBuilder,
+	const UE::NNE::FTensorDesc& TensorDesc,
+	const UE::NNE::FTensorShape& TensorShape,
+	const IResourceAccess& ResourceAccess,
+	FRDGBufferRef Buffer,
+	const FResourceMapping& ResourceMapping,
+	FRDGTextureRef Texture) const
+{
+	const ENNETensorDataType TensorDataType = TensorDesc.GetDataType();
+	const NNEDenoiserShaders::Internal::ENNEDenoiserDataType DataType = GetDenoiserShaderDataType(TensorDataType);
+	const FIntPoint BufferSize(TensorShape.GetData()[3], TensorShape.GetData()[2]);
+		
+	TMap<int32, TArray<FIntPoint>> ChannelMappingPerFrame = ResourceMapping.GetChannelMappingPerFrame(EResourceName::Output);
+	checkf(ChannelMappingPerFrame.Num() == 1, TEXT("Invalid output mapping"));
+	const TArray<FIntPoint>& ChannelMapping = ChannelMappingPerFrame.CreateConstIterator().Value();
+
+	FRDGBufferUAVRef BufferUAV = GraphBuilder.CreateUAV(FRDGBufferUAVDesc(Buffer, GetBufferFormat(TensorDataType)));
+	FRDGTextureUAVRef TextureUAV = GraphBuilder.CreateUAV(FRDGTextureUAVDesc(Texture));
+
+	IOProcessBaseHelper::AddWriteOutputPass(GraphBuilder, BufferUAV, BufferSize, TextureUAV, DataType, ChannelMapping);
+}
+
+} // namespace UE::NNEDenoiser::Private
\ No newline at end of file
diff --git a/Engine/Plugins/Experimental/NNEDenoiser/Source/NNEDenoiser/Private/NNEDenoiserIOProcessBase.h b/Engine/Plugins/Experimental/NNEDenoiser/Source/NNEDenoiser/Private/NNEDenoiserIOProcessBase.h
new file mode 100644
--- /dev/null
+++ b/Engine/Plugins/Experimental/NNEDenoiser/Source/NNEDenoiser/Private/NNEDenoiserIOProcessBase.h
@@ -0,0 +1,107 @@
+// Copyright Epic Games, Inc. All Rights Reserved.
+
+#pragma once
+
+#include "NNEDenoiserIOProcess.h"
+#include "NNEDenoiserResourceMapping.h"
+
+namespace UE::NNEDenoiser::Private
+{
+
+class FInputProcessBase : public IInputProcess
+{
+public:
+	FInputProcessBase(FResourceMappingList InputLayout) : InputLayout(InputLayout)
+	{
+
+	}
+
+	virtual ~FInputProcessBase() = default;
+
+	virtual bool PrepareAndValidate(IModelInstance& ModelInstance, FIntPoint Extent) const override;
+
+	virtual int32 NumFrames(EResourceName Name) const override;
+
+	void AddPasses(
+		FRDGBuilder& GraphBuilder,
+		TConstArrayView<UE::NNE::FTensorDesc> TensorDescs,
+		TConstArrayView<UE::NNE::FTensorShape> TensorShapes,
+		const IResourceAccess& ResourceAccess,
+		TConstArrayView<FRDGBufferRef> OutputBuffers) const override;
+
+protected:
+	virtual bool HasPreprocessInput(EResourceName TensorName, int32 FrameIdx) const
+	{
+		return false;
+	}
+
+	virtual void PreprocessInput(
+		FRDGBuilder& GraphBuilder,
+		FRDGTextureRef Texture,
+		EResourceName TensorName,
+		int32 FrameIdx,
+		FRDGTextureRef PreprocessedTexture) const
+	{
+		
+	}
+
+	virtual void WriteInputBuffer(
+		FRDGBuilder& GraphBuilder,
+		const UE::NNE::FTensorDesc& TensorDesc,
+		const UE::NNE::FTensorShape& TensorShape,
+		const IResourceAccess& ResourceAccess,
+		const FResourceMapping& ResourceMapping,
+		FRDGBufferRef Buffer) const;
+
+private:
+	FResourceMappingList InputLayout;
+};
+
+class FOutputProcessBase : public IOutputProcess
+{
+public:
+	FOutputProcessBase(FResourceMappingList OutputLayout) : OutputLayout(OutputLayout)
+	{
+
+	}
+
+	virtual ~FOutputProcessBase() = default;
+
+	virtual bool Validate(const IModelInstance& ModelInstance, FIntPoint Extent) const override;
+
+	void AddPasses(
+		FRDGBuilder& GraphBuilder,
+		TConstArrayView<UE::NNE::FTensorDesc> TensorDescs,
+		TConstArrayView<UE::NNE::FTensorShape> TensorShapes,
+		const IResourceAccess& ResourceAccess,
+		TConstArrayView<FRDGBufferRef> Buffers,
+		FRDGTextureRef OutputTexture) const override;
+
+protected:
+	virtual bool HasPostprocessOutput(EResourceName TensorName, int32 FrameIdx) const
+	{
+		return false;
+	}
+
+	virtual void ReadOutputBuffer(
+		FRDGBuilder& GraphBuilder,
+		const UE::NNE::FTensorDesc& TensorDesc,
+		const UE::NNE::FTensorShape& TensorShape,
+		const IResourceAccess& ResourceAccess,
+		FRDGBufferRef Buffer,
+		const FResourceMapping& ResourceMapping,
+		FRDGTextureRef Texture) const;
+
+	virtual void PostprocessOutput(
+		FRDGBuilder& GraphBuilder,
+		FRDGTextureRef Texture,
+		FRDGTextureRef PostprocessedTexture) const
+	{
+		
+	}
+
+private:
+	FResourceMappingList OutputLayout;
+};
+
+} // namespace UE::NNEDenoiser::Private
\ No newline at end of file
diff --git a/Engine/Plugins/Experimental/NNEDenoiser/Source/NNEDenoiser/Private/NNEDenoiserIOProcessOidn.cpp b/Engine/Plugins/Experimental/NNEDenoiser/Source/NNEDenoiser/Private/NNEDenoiserIOProcessOidn.cpp
new file mode 100644
--- /dev/null
+++ b/Engine/Plugins/Experimental/NNEDenoiser/Source/NNEDenoiser/Private/NNEDenoiserIOProcessOidn.cpp
@@ -0,0 +1,107 @@
+// Copyright Epic Games, Inc. All Rights Reserved.
+
+#include "NNEDenoiserIOProcessOidn.h"
+#include "NNEDenoiserLog.h"
+#include "NNEDenoiserParameters.h"
+#include "NNEDenoiserShadersOidnCS.h"
+#include "NNEDenoiserTransferFunction.h"
+#include "NNEDenoiserUtils.h"
+#include "RHIStaticStates.h"
+
+DECLARE_GPU_STAT_NAMED(FNNEDenoiserOidnPreOrPostprocessing, TEXT("NNEDenoiser.OidnPreOrPostprocessing"));
+
+namespace UE::NNEDenoiser::Private
+{
+
+namespace IOProcessOidnHelper
+{
+
+void AddPreOrPostProcess(
+	FRDGBuilder& GraphBuilder,
+	FRDGTextureRef InputTexture,
+	EResourceName TensorName,
+	int32 FrameIdx,
+	FRDGTextureRef OutputTexture)
+{
+	using namespace UE::NNEDenoiserShaders::Internal;
+
+	FHDRTransferFunction TransferFunction{HDRMax};
+
+	const FIntVector InputTextureSize = InputTexture->Desc.GetSize();
+	const FIntVector OutputTextureSize = OutputTexture->Desc.GetSize();
+
+	FNNEDenoiserOidnCS::FParameters *ShaderParameters = GraphBuilder.AllocParameters<FNNEDenoiserOidnCS::FParameters>();
+	ShaderParameters->InputTextureWidth = InputTextureSize.X;
+	ShaderParameters->InputTextureHeight = InputTextureSize.Y;
+	ShaderParameters->InputTexture = InputTexture;
+	ShaderParameters->OutputTextureWidth = OutputTextureSize.X;
+	ShaderParameters->OutputTextureHeight = OutputTextureSize.Y;
+	ShaderParameters->OutputTexture = GraphBuilder.CreateUAV(FRDGTextureUAVDesc(OutputTexture));
+	ShaderParameters->NormScale = TransferFunction.GetNormScale();
+	ShaderParameters->InvNormScale = TransferFunction.GetInvNormScale();
+
+	FNNEDenoiserOidnCS::FPermutationDomain PermutationVector;
+	PermutationVector.Set<FNNEDenoiserOidnCS::FNNEDenoiserInputKind>(GetInputKind(TensorName));
+
+	FIntVector ThreadGroupCount = FIntVector(
+		FMath::DivideAndRoundUp(OutputTextureSize.X, FNNEDenoiserOidnConstants::THREAD_GROUP_SIZE),
+		FMath::DivideAndRoundUp(OutputTextureSize.Y, FNNEDenoiserOidnConstants::THREAD_GROUP_SIZE),
+		1);
+
+	FGlobalShaderMap* GlobalShaderMap = GetGlobalShaderMap(GMaxRHIFeatureLevel);
+	TShaderMapRef<FNNEDenoiserOidnCS> Shader(GlobalShaderMap, PermutationVector);
+
+	RDG_EVENT_SCOPE(GraphBuilder, "NNEDenoiser.OidnPreOrPostprocessing");
+	RDG_GPU_STAT_SCOPE(GraphBuilder, FNNEDenoiserOidnPreOrPostprocessing);
+
+	FComputeShaderUtils::AddPass(
+		GraphBuilder,
+		RDG_EVENT_NAME("NNEDenoiser.OidnPreOrPostprocessing"),
+		ERDGPassFlags::Compute | ERDGPassFlags::NeverCull,
+		Shader,
+		ShaderParameters,
+		ThreadGroupCount);
+}
+
+} // namespace IOProcessOidnHelper
+
+bool FInputProcessOidn::HasPreprocessInput(EResourceName TensorName, int32 FrameIdx) const
+{
+	check(FrameIdx == 0);
+
+	switch(TensorName)
+	{
+		case EResourceName::Color:
+		case EResourceName::Albedo:
+		case EResourceName::Normal:
+			return true;
+	}
+	return false;
+}
+
+void FInputProcessOidn::PreprocessInput(
+	FRDGBuilder& GraphBuilder,
+	FRDGTextureRef Texture,
+	EResourceName TensorName,
+	int32 FrameIdx,
+	FRDGTextureRef PreprocessedTexture) const
+{
+	IOProcessOidnHelper::AddPreOrPostProcess(GraphBuilder, Texture, TensorName, FrameIdx, PreprocessedTexture);
+}
+
+bool FOutputProcessOidn::HasPostprocessOutput(EResourceName TensorName, int32 FrameIdx) const
+{
+	check(FrameIdx == 0);
+
+	return TensorName == EResourceName::Output;
+}
+
+void FOutputProcessOidn::PostprocessOutput(
+	FRDGBuilder& GraphBuilder,
+	FRDGTextureRef Texture,
+	FRDGTextureRef PostprocessedTexture) const
+{
+	IOProcessOidnHelper::AddPreOrPostProcess(GraphBuilder, Texture, EResourceName::Output, 0, PostprocessedTexture);
+}
+
+} // namespace UE::NNEDenoiser::Private
\ No newline at end of file
diff --git a/Engine/Plugins/Experimental/NNEDenoiser/Source/NNEDenoiser/Private/NNEDenoiserIOProcessOidn.h b/Engine/Plugins/Experimental/NNEDenoiser/Source/NNEDenoiser/Private/NNEDenoiserIOProcessOidn.h
new file mode 100644
--- /dev/null
+++ b/Engine/Plugins/Experimental/NNEDenoiser/Source/NNEDenoiser/Private/NNEDenoiserIOProcessOidn.h
@@ -0,0 +1,50 @@
+// Copyright Epic Games, Inc. All Rights Reserved.
+
+#pragma once
+
+#include "NNEDenoiserIOProcessBase.h"
+
+namespace UE::NNEDenoiser::Private
+{
+
+class FInputProcessOidn : public FInputProcessBase
+{
+public:
+	FInputProcessOidn(FResourceMappingList InputLayout) : FInputProcessBase(InputLayout)
+	{
+		
+	}
+
+	virtual ~FInputProcessOidn() = default;
+
+protected:
+	virtual bool HasPreprocessInput(EResourceName TensorName, int32 FrameIdx) const override;
+
+	virtual void PreprocessInput(
+		FRDGBuilder& GraphBuilder,
+		FRDGTextureRef Texture,
+		EResourceName TensorName,
+		int32 FrameIdx,
+		FRDGTextureRef PreprocessedTexture) const override;
+};
+
+class FOutputProcessOidn : public FOutputProcessBase
+{
+public:
+	FOutputProcessOidn(FResourceMappingList OutputLayout) : FOutputProcessBase(OutputLayout)
+	{
+		
+	}
+
+	virtual ~FOutputProcessOidn() = default;
+
+protected:
+	virtual bool HasPostprocessOutput(EResourceName TensorName, int32 FrameIdx) const override;
+
+	virtual void PostprocessOutput(
+		FRDGBuilder& GraphBuilder,
+		FRDGTextureRef Texture,
+		FRDGTextureRef PostprocessedTexture) const override;
+};
+
+} // namespace UE::NNEDenoiser::Private
\ No newline at end of file
diff --git a/Engine/Plugins/Experimental/NNEDenoiser/Source/NNEDenoiser/Private/NNEDenoiserLog.h b/Engine/Plugins/Experimental/NNEDenoiser/Source/NNEDenoiser/Private/NNEDenoiserLog.h
new file mode 100644
--- /dev/null
+++ b/Engine/Plugins/Experimental/NNEDenoiser/Source/NNEDenoiser/Private/NNEDenoiserLog.h
@@ -0,0 +1,8 @@
+// Copyright Epic Games, Inc. All Rights Reserved.
+
+#pragma once
+
+#include "CoreMinimal.h"
+#include "Logging/LogMacros.h"
+
+DECLARE_LOG_CATEGORY_EXTERN(LogNNEDenoiser, Log, All);
\ No newline at end of file
diff --git a/Engine/Plugins/Experimental/NNEDenoiser/Source/NNEDenoiser/Private/NNEDenoiserModelInstance.h b/Engine/Plugins/Experimental/NNEDenoiser/Source/NNEDenoiser/Private/NNEDenoiserModelInstance.h
new file mode 100644
--- /dev/null
+++ b/Engine/Plugins/Experimental/NNEDenoiser/Source/NNEDenoiser/Private/NNEDenoiserModelInstance.h
@@ -0,0 +1,16 @@
+// Copyright Epic Games, Inc. All Rights Reserved.
+
+#pragma once
+
+#include "NNERuntimeRDG.h"
+
+namespace UE::NNEDenoiser::Private
+{
+
+class IModelInstance : public NNE::IModelInstanceRDG
+{
+public:
+	virtual ~IModelInstance() = default;
+};
+
+}
\ No newline at end of file
diff --git a/Engine/Plugins/Experimental/NNEDenoiser/Source/NNEDenoiser/Private/NNEDenoiserModelInstanceCPU.cpp b/Engine/Plugins/Experimental/NNEDenoiser/Source/NNEDenoiser/Private/NNEDenoiserModelInstanceCPU.cpp
new file mode 100644
--- /dev/null
+++ b/Engine/Plugins/Experimental/NNEDenoiser/Source/NNEDenoiser/Private/NNEDenoiserModelInstanceCPU.cpp
@@ -0,0 +1,164 @@
+// Copyright Epic Games, Inc. All Rights Reserved.
+
+#include "NNEDenoiserModelInstanceCPU.h"
+#include "NNE.h"
+#include "NNEDenoiserLog.h"
+#include "NNEDenoiserUtils.h"
+#include "NNEModelData.h"
+#include "NNERuntimeCPU.h"
+#include "RenderGraphBuilder.h"
+#include "RHICommandList.h"
+#include "RHIResources.h"
+#include "RHITypes.h"
+#include "UObject/UObjectGlobals.h"
+
+BEGIN_SHADER_PARAMETER_STRUCT(FNNEDenoiserModelInstanceCPUTextureParameters, )
+	RDG_BUFFER_ACCESS_ARRAY(InputBuffers)
+	RDG_BUFFER_ACCESS_ARRAY(OutputBuffers)
+END_SHADER_PARAMETER_STRUCT()
+
+namespace UE::NNEDenoiser::Private
+{
+
+	static const FString DefaultRuntimeCPUName = TEXT("NNERuntimeORTCpu");
+
+	TUniquePtr<FModelInstanceCPU> FModelInstanceCPU::Make(UNNEModelData& ModelData, const FString &RuntimeNameOverride)
+	{
+		const FString RuntimeCPUName = RuntimeNameOverride.IsEmpty() ? DefaultRuntimeCPUName : RuntimeNameOverride;
+
+		// Create the model
+		TWeakInterfacePtr<INNERuntimeCPU> RuntimeCPU = UE::NNE::GetRuntime<INNERuntimeCPU>(RuntimeCPUName);
+		if (!RuntimeCPU.IsValid())
+		{
+			UE_LOG(LogNNEDenoiser, Error, TEXT("No CPU runtime '%s' found. Valid CPU runtimes are: "), *RuntimeCPUName);
+			for (TWeakInterfacePtr<INNERuntime> Runtime : UE::NNE::GetAllRuntimes())
+			{
+				if (Runtime.IsValid() && UE::NNE::GetRuntime<INNERuntimeCPU>(Runtime->GetRuntimeName()).IsValid())
+				{
+					UE_LOG(LogNNEDenoiser, Error, TEXT("- %s"), *Runtime->GetRuntimeName());
+				}
+			}
+			return {};
+		}
+
+		TSharedPtr<UE::NNE::IModelCPU> Model = RuntimeCPU->CreateModelCPU(ModelData);
+		if (!Model.IsValid())
+		{
+			UE_LOG(LogNNEDenoiser, Error, TEXT("Could not create model using %s"), *RuntimeCPUName);
+			return {};
+		}
+
+		TSharedPtr<UE::NNE::IModelInstanceCPU> ModelInstance = Model->CreateModelInstanceCPU();
+		if (!ModelInstance.IsValid())
+		{
+			UE_LOG(LogNNEDenoiser, Error, TEXT("Could not create model instance using %s"), *RuntimeCPUName);
+			return {};
+		}
+
+		UE_LOG(LogNNEDenoiser, Log, TEXT("NNEDenoiserCPU: Creaded model instance from %s using %s"), *ModelData.GetFileId().ToString(), *RuntimeCPUName);
+
+		return MakeUnique<FModelInstanceCPU>(ModelInstance.ToSharedRef());
+	}
+
+	FModelInstanceCPU::FModelInstanceCPU(TSharedRef<UE::NNE::IModelInstanceCPU> ModelInstance) :
+			ModelInstance(ModelInstance)
+	{
+
+	}
+
+	FModelInstanceCPU::~FModelInstanceCPU()
+	{
+		
+	}
+
+	TConstArrayView<NNE::FTensorDesc> FModelInstanceCPU::GetInputTensorDescs() const
+	{
+		return ModelInstance->GetInputTensorDescs();
+	}
+
+	TConstArrayView<NNE::FTensorDesc> FModelInstanceCPU::GetOutputTensorDescs() const
+	{
+		return ModelInstance->GetOutputTensorDescs();
+	}
+
+	TConstArrayView<NNE::FTensorShape> FModelInstanceCPU::GetInputTensorShapes() const
+	{
+		return ModelInstance->GetInputTensorShapes();
+	}
+
+	TConstArrayView<NNE::FTensorShape> FModelInstanceCPU::GetOutputTensorShapes() const
+	{
+		return ModelInstance->GetOutputTensorShapes();
+	}
+
+	int32 FModelInstanceCPU::SetInputTensorShapes(TConstArrayView<NNE::FTensorShape> InInputShapes)
+	{
+		return ModelInstance->SetInputTensorShapes(InInputShapes);
+	}
+
+	int32 FModelInstanceCPU::EnqueueRDG(FRDGBuilder &GraphBuilder, TConstArrayView<NNE::FTensorBindingRDG> Inputs, TConstArrayView<NNE::FTensorBindingRDG> Outputs)
+	{
+		FNNEDenoiserModelInstanceCPUTextureParameters* DenoiserParameters = GraphBuilder.AllocParameters<FNNEDenoiserModelInstanceCPUTextureParameters>();
+		for (const NNE::FTensorBindingRDG& Binding : Inputs)
+		{
+			DenoiserParameters->InputBuffers.Emplace(Binding.Buffer, ERHIAccess::CopySrc);
+		}
+		for (const NNE::FTensorBindingRDG& Binding : Outputs)
+		{
+			DenoiserParameters->OutputBuffers.Emplace(Binding.Buffer, ERHIAccess::CopySrc);
+		}
+		
+		GraphBuilder.AddPass(RDG_EVENT_NAME("NNEDenoiser.DenoiseCPU"), DenoiserParameters, ERDGPassFlags::Readback,
+		[this, DenoiserParameters](FRHICommandListImmediate& RHICmdList)
+		{
+			SCOPED_NAMED_EVENT_TEXT("FModelInstanceCPU::DenoisePass", FColor::Magenta);
+			
+#if WITH_EDITOR
+			// NOTE: the time will include the transfer from GPU to CPU which will include waiting for the GPU pipeline to complete
+			uint64 FilterExecuteTime = 0;
+			FilterExecuteTime -= FPlatformTime::Cycles64();
+#endif
+
+			ScratchInputBuffers.SetNum(DenoiserParameters->InputBuffers.Num());
+			TArray<NNE::FTensorBindingCPU> InputBindings;
+			for (int32 Idx = 0; Idx < DenoiserParameters->InputBuffers.Num(); Idx++)
+			{
+				ScratchInputBuffers[Idx].SetNumUninitialized(DenoiserParameters->InputBuffers[Idx].GetBuffer()->GetRHI()->GetSize());
+
+				InputBindings.Emplace(NNE::FTensorBindingCPU{(void *)ScratchInputBuffers[Idx].GetData(), (uint64)ScratchInputBuffers[Idx].Num()});
+			}
+
+			ScratchOutputBuffers.SetNum(DenoiserParameters->OutputBuffers.Num());
+			TArray<NNE::FTensorBindingCPU> OutputBindings;
+			for (int32 Idx = 0; Idx < DenoiserParameters->OutputBuffers.Num(); Idx++)
+			{
+				ScratchOutputBuffers[Idx].SetNumUninitialized(DenoiserParameters->OutputBuffers[Idx].GetBuffer()->GetRHI()->GetSize());
+
+				OutputBindings.Emplace(NNE::FTensorBindingCPU{(void *)ScratchOutputBuffers[Idx].GetData(), (uint64)ScratchOutputBuffers[Idx].Num()});
+			}
+
+			for (int32 Idx = 0; Idx < DenoiserParameters->InputBuffers.Num(); Idx++)
+			{
+				FRHIBuffer* Buffer = DenoiserParameters->InputBuffers[Idx].GetBuffer()->GetRHI();
+				CopyBufferFromGPUToCPU(RHICmdList, Buffer, Buffer->GetSize(), ScratchInputBuffers[Idx]);
+			}
+
+			ModelInstance->RunSync(InputBindings, OutputBindings);
+
+			for (int32 Idx = 0; Idx < DenoiserParameters->OutputBuffers.Num(); Idx++)
+			{
+				FRHIBuffer* Buffer = DenoiserParameters->OutputBuffers[Idx].GetBuffer()->GetRHI();
+				CopyBufferFromCPUToGPU(RHICmdList, ScratchOutputBuffers[Idx], Buffer->GetSize(), Buffer);
+			}
+
+#if WITH_EDITOR
+		FilterExecuteTime += FPlatformTime::Cycles64();
+		const double FilterExecuteTimeMS = 1000.0 * FPlatformTime::ToSeconds64(FilterExecuteTime);
+		UE_LOG(LogNNEDenoiser, Log, TEXT("Denoised on CPU in %.2f ms"), FilterExecuteTimeMS);
+#endif
+		});
+
+		return 0;
+	}
+
+}
\ No newline at end of file
diff --git a/Engine/Plugins/Experimental/NNEDenoiser/Source/NNEDenoiser/Private/NNEDenoiserModelInstanceCPU.h b/Engine/Plugins/Experimental/NNEDenoiser/Source/NNEDenoiser/Private/NNEDenoiserModelInstanceCPU.h
new file mode 100644
--- /dev/null
+++ b/Engine/Plugins/Experimental/NNEDenoiser/Source/NNEDenoiser/Private/NNEDenoiserModelInstanceCPU.h
@@ -0,0 +1,46 @@
+// Copyright Epic Games, Inc. All Rights Reserved.
+
+#pragma once
+
+#include "NNEDenoiserModelInstance.h"
+#include "NNEModelData.h"
+
+class FRDGBuilder;
+
+namespace UE::NNE
+{
+	class IModelInstanceCPU;
+}
+
+namespace UE::NNEDenoiser::Private
+{
+
+class FModelInstanceCPU : public IModelInstance
+{
+public:
+	explicit FModelInstanceCPU(TSharedRef<UE::NNE::IModelInstanceCPU> ModelInstance);
+
+	static TUniquePtr<FModelInstanceCPU> Make(UNNEModelData& ModelData, const FString& RuntimeNameOverride = {});
+
+	virtual ~FModelInstanceCPU();
+
+	virtual TConstArrayView<NNE::FTensorDesc> GetInputTensorDescs() const override;
+
+	virtual TConstArrayView<NNE::FTensorDesc> GetOutputTensorDescs() const override;
+
+	virtual TConstArrayView<NNE::FTensorShape> GetInputTensorShapes() const override;
+
+	virtual TConstArrayView<NNE::FTensorShape> GetOutputTensorShapes() const override;
+
+	virtual int32 SetInputTensorShapes(TConstArrayView<NNE::FTensorShape> InInputShapes) override;
+
+	virtual int32 EnqueueRDG(FRDGBuilder &GraphBuilder, TConstArrayView<NNE::FTensorBindingRDG> Inputs, TConstArrayView<NNE::FTensorBindingRDG> Outputs) override;
+
+private:
+	TSharedRef<UE::NNE::IModelInstanceCPU> ModelInstance;
+
+	TArray<TArray<uint8>> ScratchInputBuffers;
+	TArray<TArray<uint8>> ScratchOutputBuffers;
+};
+	
+}
\ No newline at end of file
diff --git a/Engine/Plugins/Experimental/NNEDenoiser/Source/NNEDenoiser/Private/NNEDenoiserModelInstanceGPU.cpp b/Engine/Plugins/Experimental/NNEDenoiser/Source/NNEDenoiser/Private/NNEDenoiserModelInstanceGPU.cpp
new file mode 100644
--- /dev/null
+++ b/Engine/Plugins/Experimental/NNEDenoiser/Source/NNEDenoiser/Private/NNEDenoiserModelInstanceGPU.cpp
@@ -0,0 +1,170 @@
+// Copyright Epic Games, Inc. All Rights Reserved.
+
+#include "NNEDenoiserModelInstanceGPU.h"
+#include "NNE.h"
+#include "NNEDenoiserLog.h"
+#include "NNEDenoiserUtils.h"
+#include "NNEModelData.h"
+#include "NNERuntimeGPU.h"
+#include "PathTracingDenoiser.h"
+#include "RenderGraphBuilder.h"
+#include "RHICommandList.h"
+#include "RHIResources.h"
+#include "RHITypes.h"
+#include "UObject/UObjectGlobals.h"
+
+DECLARE_GPU_STAT_NAMED(FModelInstanceGPU, TEXT("NNEDenoiser.ModelInstanceGPU"));
+
+BEGIN_SHADER_PARAMETER_STRUCT(FNNEDenoiserModelInstanceGPUTextureParameters, )
+	RDG_BUFFER_ACCESS_ARRAY(InputBuffers)
+	RDG_BUFFER_ACCESS_ARRAY(OutputBuffers)
+END_SHADER_PARAMETER_STRUCT()
+
+namespace UE::NNEDenoiser::Private
+{
+
+	static const FString DefaultRuntimeGPUName = TEXT("NNERuntimeORTDml");
+
+	TUniquePtr<FModelInstanceGPU> FModelInstanceGPU::Make(UNNEModelData& ModelData, const FString& RuntimeNameOverride)
+	{
+		const FString RuntimeGPUName = RuntimeNameOverride.IsEmpty() ? DefaultRuntimeGPUName : RuntimeNameOverride;
+
+		// Create the model
+		TWeakInterfacePtr<INNERuntimeGPU> RuntimeGPU = UE::NNE::GetRuntime<INNERuntimeGPU>(RuntimeGPUName);
+		if (!RuntimeGPU.IsValid())
+		{
+			UE_LOG(LogNNEDenoiser, Error, TEXT("No GPU runtime '%s' found. Valid GPU runtimes are: "), *RuntimeGPUName);
+			for (TWeakInterfacePtr<INNERuntime> Runtime : UE::NNE::GetAllRuntimes())
+			{
+				if (Runtime.IsValid() && UE::NNE::GetRuntime<INNERuntimeGPU>(Runtime->GetRuntimeName()).IsValid())
+				{
+					UE_LOG(LogNNEDenoiser, Error, TEXT("- %s"), *Runtime->GetRuntimeName());
+				}
+			}
+			return {};
+		}
+
+		TSharedPtr<UE::NNE::IModelGPU> Model = RuntimeGPU->CreateModelGPU(ModelData);
+		if (!Model.IsValid())
+		{
+			UE_LOG(LogNNEDenoiser, Error, TEXT("Could not create model using %s"), *RuntimeGPUName);
+			return {};
+		}
+
+		TSharedPtr<UE::NNE::IModelInstanceGPU> ModelInstance = Model->CreateModelInstanceGPU();
+		if (!ModelInstance.IsValid())
+		{
+			UE_LOG(LogNNEDenoiser, Error, TEXT("Could not create model instance using %s"), *RuntimeGPUName);
+			return {};
+		}
+
+		UE_LOG(LogNNEDenoiser, Log, TEXT("NNEDenoiserGPU: Creaded model instance from %s using %s"), *ModelData.GetFileId().ToString(), *RuntimeGPUName);
+
+		return MakeUnique<FModelInstanceGPU>(ModelInstance.ToSharedRef());
+	}
+
+	FModelInstanceGPU::FModelInstanceGPU(TSharedRef<UE::NNE::IModelInstanceGPU> ModelInstance) :
+		ModelInstance(ModelInstance)
+	{
+
+	}
+
+	FModelInstanceGPU::~FModelInstanceGPU()
+	{
+		
+	}
+
+	TConstArrayView<NNE::FTensorDesc> FModelInstanceGPU::GetInputTensorDescs() const
+	{
+		return ModelInstance->GetInputTensorDescs();
+	}
+
+	TConstArrayView<NNE::FTensorDesc> FModelInstanceGPU::GetOutputTensorDescs() const
+	{
+		return ModelInstance->GetOutputTensorDescs();
+	}
+
+	TConstArrayView<NNE::FTensorShape> FModelInstanceGPU::GetInputTensorShapes() const
+	{
+		return ModelInstance->GetInputTensorShapes();
+	}
+
+	TConstArrayView<NNE::FTensorShape> FModelInstanceGPU::GetOutputTensorShapes() const
+	{
+		return ModelInstance->GetOutputTensorShapes();
+	}
+
+	int32 FModelInstanceGPU::SetInputTensorShapes(TConstArrayView<NNE::FTensorShape> InInputShapes)
+	{
+		return ModelInstance->SetInputTensorShapes(InInputShapes);
+	}
+
+	int32 FModelInstanceGPU::EnqueueRDG(FRDGBuilder &GraphBuilder, TConstArrayView<NNE::FTensorBindingRDG> Inputs, TConstArrayView<NNE::FTensorBindingRDG> Outputs)
+	{
+		FNNEDenoiserModelInstanceGPUTextureParameters* DenoiserParameters = GraphBuilder.AllocParameters<FNNEDenoiserModelInstanceGPUTextureParameters>();
+		for (const NNE::FTensorBindingRDG& Binding : Inputs)
+		{
+			DenoiserParameters->InputBuffers.Emplace(Binding.Buffer, ERHIAccess::CopySrc);
+		}
+		for (const NNE::FTensorBindingRDG& Binding : Outputs)
+		{
+			DenoiserParameters->OutputBuffers.Emplace(Binding.Buffer, ERHIAccess::CopySrc);
+		}
+
+		RDG_EVENT_SCOPE(GraphBuilder, "NNEDenoiser.ModelInstanceGPU");
+		RDG_GPU_STAT_SCOPE(GraphBuilder, FModelInstanceGPU);
+		
+		GraphBuilder.AddPass(RDG_EVENT_NAME("NNEDenoiser.DenoiseGPU"), DenoiserParameters, ERDGPassFlags::Readback,
+		[this, DenoiserParameters](FRHICommandListImmediate& RHICmdList)
+		{
+			SCOPED_NAMED_EVENT_TEXT("FModelInstanceGPU::DenoisePass", FColor::Magenta);
+			
+#if WITH_EDITOR
+			// NOTE: the time will include the transfer from GPU to CPU which will include waiting for the GPU pipeline to complete
+			uint64 FilterExecuteTime = 0;
+			FilterExecuteTime -= FPlatformTime::Cycles64();
+#endif
+
+			ScratchInputBuffers.SetNum(DenoiserParameters->InputBuffers.Num());
+			TArray<NNE::FTensorBindingGPU> InputBindings;
+			for (int32 Idx = 0; Idx < DenoiserParameters->InputBuffers.Num(); Idx++)
+			{
+				ScratchInputBuffers[Idx].SetNumUninitialized(DenoiserParameters->InputBuffers[Idx].GetBuffer()->GetRHI()->GetSize());
+
+				InputBindings.Emplace(NNE::FTensorBindingGPU{(void *)ScratchInputBuffers[Idx].GetData(), (uint64)ScratchInputBuffers[Idx].Num()});
+			}
+
+			ScratchOutputBuffers.SetNum(DenoiserParameters->OutputBuffers.Num());
+			TArray<NNE::FTensorBindingGPU> OutputBindings;
+			for (int32 Idx = 0; Idx < DenoiserParameters->OutputBuffers.Num(); Idx++)
+			{
+				ScratchOutputBuffers[Idx].SetNumUninitialized(DenoiserParameters->OutputBuffers[Idx].GetBuffer()->GetRHI()->GetSize());
+
+				OutputBindings.Emplace(NNE::FTensorBindingGPU{(void *)ScratchOutputBuffers[Idx].GetData(), (uint64)ScratchOutputBuffers[Idx].Num()});
+			}
+
+			for (int32 Idx = 0; Idx < DenoiserParameters->InputBuffers.Num(); Idx++)
+			{
+				FRHIBuffer* Buffer = DenoiserParameters->InputBuffers[Idx].GetBuffer()->GetRHI();
+				CopyBufferFromGPUToCPU(RHICmdList, Buffer, Buffer->GetSize(), ScratchInputBuffers[Idx]);
+			}
+
+			ModelInstance->RunSync(InputBindings, OutputBindings);
+
+			for (int32 Idx = 0; Idx < DenoiserParameters->OutputBuffers.Num(); Idx++)
+			{
+				FRHIBuffer* Buffer = DenoiserParameters->OutputBuffers[Idx].GetBuffer()->GetRHI();
+				CopyBufferFromCPUToGPU(RHICmdList, ScratchOutputBuffers[Idx], Buffer->GetSize(), Buffer);
+			}
+
+#if WITH_EDITOR
+		FilterExecuteTime += FPlatformTime::Cycles64();
+		const double FilterExecuteTimeMS = 1000.0 * FPlatformTime::ToSeconds64(FilterExecuteTime);
+		UE_LOG(LogNNEDenoiser, Log, TEXT("Denoised on GPU in %.2f ms"), FilterExecuteTimeMS);
+#endif
+		});
+
+		return 0;
+	}
+
+}
\ No newline at end of file
diff --git a/Engine/Plugins/Experimental/NNEDenoiser/Source/NNEDenoiser/Private/NNEDenoiserModelInstanceGPU.h b/Engine/Plugins/Experimental/NNEDenoiser/Source/NNEDenoiser/Private/NNEDenoiserModelInstanceGPU.h
new file mode 100644
--- /dev/null
+++ b/Engine/Plugins/Experimental/NNEDenoiser/Source/NNEDenoiser/Private/NNEDenoiserModelInstanceGPU.h
@@ -0,0 +1,44 @@
+// Copyright Epic Games, Inc. All Rights Reserved.
+
+#pragma once
+
+#include "NNEDenoiserModelInstance.h"
+#include "NNEModelData.h"
+
+namespace UE::NNE
+{
+	class IModelInstanceGPU;
+}
+
+namespace UE::NNEDenoiser::Private
+{
+
+class FModelInstanceGPU : public IModelInstance
+{
+public:
+	explicit FModelInstanceGPU(TSharedRef<UE::NNE::IModelInstanceGPU> ModelInstance);
+
+	static TUniquePtr<FModelInstanceGPU> Make(UNNEModelData& ModelData, const FString& RuntimeNameOverride = {});
+
+	virtual ~FModelInstanceGPU();
+
+	virtual TConstArrayView<NNE::FTensorDesc> GetInputTensorDescs() const override;
+
+	virtual TConstArrayView<NNE::FTensorDesc> GetOutputTensorDescs() const override;
+
+	virtual TConstArrayView<NNE::FTensorShape> GetInputTensorShapes() const override;
+
+	virtual TConstArrayView<NNE::FTensorShape> GetOutputTensorShapes() const override;
+
+	virtual int32 SetInputTensorShapes(TConstArrayView<NNE::FTensorShape> InInputShapes) override;
+
+	virtual int32 EnqueueRDG(FRDGBuilder &GraphBuilder, TConstArrayView<NNE::FTensorBindingRDG> Inputs, TConstArrayView<NNE::FTensorBindingRDG> Outputs) override;
+
+private:
+	TSharedRef<UE::NNE::IModelInstanceGPU> ModelInstance;
+
+	TArray<TArray<uint8>> ScratchInputBuffers;
+	TArray<TArray<uint8>> ScratchOutputBuffers;
+};
+	
+}
\ No newline at end of file
diff --git a/Engine/Plugins/Experimental/NNEDenoiser/Source/NNEDenoiser/Private/NNEDenoiserModelInstanceRDG.cpp b/Engine/Plugins/Experimental/NNEDenoiser/Source/NNEDenoiser/Private/NNEDenoiserModelInstanceRDG.cpp
new file mode 100644
--- /dev/null
+++ b/Engine/Plugins/Experimental/NNEDenoiser/Source/NNEDenoiser/Private/NNEDenoiserModelInstanceRDG.cpp
@@ -0,0 +1,100 @@
+// Copyright Epic Games, Inc. All Rights Reserved.
+
+#include "NNEDenoiserModelInstanceRDG.h"
+#include "NNE.h"
+#include "NNEDenoiserLog.h"
+#include "NNEDenoiserUtils.h"
+#include "NNEModelData.h"
+#include "NNERuntimeRDG.h"
+#include "PathTracingDenoiser.h"
+#include "RenderGraphBuilder.h"
+#include "RHICommandList.h"
+#include "RHIResources.h"
+#include "RHITypes.h"
+#include "UObject/UObjectGlobals.h"
+
+namespace UE::NNEDenoiser::Private
+{
+
+	static const FString DefaultRuntimeRDGName = TEXT("NNERuntimeRDGDml");
+
+	TUniquePtr<FModelInstanceRDG> FModelInstanceRDG::Make(UNNEModelData& ModelData, const FString& RuntimeNameOverride)
+	{
+		FString RuntimeRDGName = RuntimeNameOverride.IsEmpty() ? DefaultRuntimeRDGName : RuntimeNameOverride;
+		
+		// Create the model
+		TWeakInterfacePtr<INNERuntimeRDG> RuntimeRDG = UE::NNE::GetRuntime<INNERuntimeRDG>(RuntimeRDGName);
+		if (!RuntimeRDG.IsValid())
+		{
+			UE_LOG(LogNNEDenoiser, Error, TEXT("No RDG runtime '%s' found. Valid RDG runtimes are: "), *RuntimeRDGName);
+			for (TWeakInterfacePtr<INNERuntime> Runtime : UE::NNE::GetAllRuntimes())
+			{
+				if (Runtime.IsValid() && UE::NNE::GetRuntime<INNERuntimeRDG>(Runtime->GetRuntimeName()).IsValid())
+				{
+					UE_LOG(LogNNEDenoiser, Error, TEXT("- %s"), *Runtime->GetRuntimeName());
+				}
+			}
+			return {};
+		}
+
+		TSharedPtr<UE::NNE::IModelRDG> Model = RuntimeRDG->CreateModelRDG(ModelData);
+		if (!Model.IsValid())
+		{
+			UE_LOG(LogNNEDenoiser, Error, TEXT("Could not create model using %s"), *RuntimeRDGName);
+			return {};
+		}
+
+		TSharedPtr<UE::NNE::IModelInstanceRDG> ModelInstance = Model->CreateModelInstanceRDG();
+		if (!ModelInstance.IsValid())
+		{
+			UE_LOG(LogNNEDenoiser, Error, TEXT("Could not create model instance using %s"), *RuntimeRDGName);
+			return {};
+		}
+
+		UE_LOG(LogNNEDenoiser, Log, TEXT("NNEDenoiserRDG: Creaded model instance from %s using %s"), *ModelData.GetFileId().ToString(), *RuntimeRDGName);
+
+		return MakeUnique<FModelInstanceRDG>(ModelInstance.ToSharedRef());
+	}
+
+	FModelInstanceRDG::FModelInstanceRDG(TSharedRef<UE::NNE::IModelInstanceRDG> ModelInstance) :
+		ModelInstance(ModelInstance)
+	{
+
+	}
+
+	FModelInstanceRDG::~FModelInstanceRDG()
+	{
+		
+	}
+
+	TConstArrayView<NNE::FTensorDesc> FModelInstanceRDG::GetInputTensorDescs() const
+	{
+		return ModelInstance->GetInputTensorDescs();
+	}
+
+	TConstArrayView<NNE::FTensorDesc> FModelInstanceRDG::GetOutputTensorDescs() const
+	{
+		return ModelInstance->GetOutputTensorDescs();
+	}
+
+	TConstArrayView<NNE::FTensorShape> FModelInstanceRDG::GetInputTensorShapes() const
+	{
+		return ModelInstance->GetInputTensorShapes();
+	}
+
+	TConstArrayView<NNE::FTensorShape> FModelInstanceRDG::GetOutputTensorShapes() const
+	{
+		return ModelInstance->GetOutputTensorShapes();
+	}
+
+	int32 FModelInstanceRDG::SetInputTensorShapes(TConstArrayView<NNE::FTensorShape> InInputShapes)
+	{
+		return ModelInstance->SetInputTensorShapes(InInputShapes);
+	}
+
+	int32 FModelInstanceRDG::EnqueueRDG(FRDGBuilder &GraphBuilder, TConstArrayView<NNE::FTensorBindingRDG> Inputs, TConstArrayView<NNE::FTensorBindingRDG> Outputs)
+	{
+		return ModelInstance->EnqueueRDG(GraphBuilder, Inputs, Outputs);
+	}
+
+}
\ No newline at end of file
diff --git a/Engine/Plugins/Experimental/NNEDenoiser/Source/NNEDenoiser/Private/NNEDenoiserModelInstanceRDG.h b/Engine/Plugins/Experimental/NNEDenoiser/Source/NNEDenoiser/Private/NNEDenoiserModelInstanceRDG.h
new file mode 100644
--- /dev/null
+++ b/Engine/Plugins/Experimental/NNEDenoiser/Source/NNEDenoiser/Private/NNEDenoiserModelInstanceRDG.h
@@ -0,0 +1,41 @@
+// Copyright Epic Games, Inc. All Rights Reserved.
+
+#pragma once
+
+#include "NNEDenoiserModelInstance.h"
+#include "NNEModelData.h"
+
+namespace UE::NNE
+{
+	class IModelInstanceRDG;
+}
+
+namespace UE::NNEDenoiser::Private
+{
+
+class FModelInstanceRDG : public IModelInstance
+{
+public:
+	explicit FModelInstanceRDG(TSharedRef<UE::NNE::IModelInstanceRDG> ModelInstance);
+
+	static TUniquePtr<FModelInstanceRDG> Make(UNNEModelData& ModelData, const FString& RuntimeNameOverride = {});
+
+	virtual ~FModelInstanceRDG();
+
+	virtual TConstArrayView<NNE::FTensorDesc> GetInputTensorDescs() const override;
+
+	virtual TConstArrayView<NNE::FTensorDesc> GetOutputTensorDescs() const override;
+
+	virtual TConstArrayView<NNE::FTensorShape> GetInputTensorShapes() const override;
+
+	virtual TConstArrayView<NNE::FTensorShape> GetOutputTensorShapes() const override;
+
+	virtual int32 SetInputTensorShapes(TConstArrayView<NNE::FTensorShape> InInputShapes) override;
+
+	virtual int32 EnqueueRDG(FRDGBuilder &GraphBuilder, TConstArrayView<NNE::FTensorBindingRDG> Inputs, TConstArrayView<NNE::FTensorBindingRDG> Outputs) override;
+
+private:
+	TSharedRef<UE::NNE::IModelInstanceRDG> ModelInstance;
+};
+	
+}
\ No newline at end of file
diff --git a/Engine/Plugins/Experimental/NNEDenoiser/Source/NNEDenoiser/Private/NNEDenoiserParameters.h b/Engine/Plugins/Experimental/NNEDenoiser/Source/NNEDenoiser/Private/NNEDenoiserParameters.h
new file mode 100644
--- /dev/null
+++ b/Engine/Plugins/Experimental/NNEDenoiser/Source/NNEDenoiser/Private/NNEDenoiserParameters.h
@@ -0,0 +1,15 @@
+// Copyright Epic Games, Inc. All Rights Reserved.
+
+#pragma once
+
+#include "HAL/Platform.h"
+
+namespace UE::NNEDenoiser::Private
+{
+	
+	struct FParameters
+	{
+		int32 TileMinimumOverlap = 20;
+	};
+
+} // namespace UE::NNEDenoiser::Private
\ No newline at end of file
diff --git a/Engine/Plugins/Experimental/NNEDenoiser/Source/NNEDenoiser/Private/NNEDenoiserPathTracingDenoiser.cpp b/Engine/Plugins/Experimental/NNEDenoiser/Source/NNEDenoiser/Private/NNEDenoiserPathTracingDenoiser.cpp
new file mode 100644
--- /dev/null
+++ b/Engine/Plugins/Experimental/NNEDenoiser/Source/NNEDenoiser/Private/NNEDenoiserPathTracingDenoiser.cpp
@@ -0,0 +1,27 @@
+// Copyright Epic Games, Inc. All Rights Reserved.
+
+#include "NNEDenoiserPathTracingDenoiser.h"
+#include "NNEDenoiserGenericDenoiser.h"
+#include "NNEDenoiserHistory.h"
+#include "NNEDenoiserLog.h"
+#include "SceneView.h"
+
+namespace UE::NNEDenoiser::Private
+{
+
+using UE::Renderer::Private::IPathTracingDenoiser;
+
+FPathTracingDenoiser::FPathTracingDenoiser(TUniquePtr<FGenericDenoiser> Denoiser) : Denoiser(MoveTemp(Denoiser))
+{
+
+}
+
+void FPathTracingDenoiser::AddPasses(FRDGBuilder& GraphBuilder, const FSceneView& View, const FInputs& Inputs) const
+{
+	TUniquePtr<FHistory> History = Denoiser->AddPasses(GraphBuilder, Inputs.ColorTex, Inputs.AlbedoTex, Inputs.NormalTex, Inputs.OutputTex, nullptr, View.GPUMask, nullptr);
+	
+	// note: can't store history
+	check(!History.IsValid()); 
+}
+
+} // namespace UE::NNEDenoiser::Private
\ No newline at end of file
diff --git a/Engine/Plugins/Experimental/NNEDenoiser/Source/NNEDenoiser/Private/NNEDenoiserPathTracingDenoiser.h b/Engine/Plugins/Experimental/NNEDenoiser/Source/NNEDenoiser/Private/NNEDenoiserPathTracingDenoiser.h
new file mode 100644
--- /dev/null
+++ b/Engine/Plugins/Experimental/NNEDenoiser/Source/NNEDenoiser/Private/NNEDenoiserPathTracingDenoiser.h
@@ -0,0 +1,31 @@
+// Copyright Epic Games, Inc. All Rights Reserved.
+
+#pragma once
+
+#include "PathTracingDenoiser.h"
+#include "Templates/UniquePtr.h"
+
+namespace UE::NNEDenoiser::Private
+{
+
+class FGenericDenoiser;
+
+using UE::Renderer::Private::IPathTracingDenoiser;
+
+class FPathTracingDenoiser final : public IPathTracingDenoiser
+{
+
+public:
+	explicit FPathTracingDenoiser(TUniquePtr<FGenericDenoiser> Denoiser);
+
+	~FPathTracingDenoiser() = default;
+
+	bool NeedTextureCreateExtraFlags() const override { return true; }
+
+	void AddPasses(FRDGBuilder& GraphBuilder, const FSceneView& View, const FInputs& Inputs) const override;
+
+private:
+	TUniquePtr<FGenericDenoiser> Denoiser;
+};
+
+} // namespace UE::NNEDenoiser::Private
\ No newline at end of file
diff --git a/Engine/Plugins/Experimental/NNEDenoiser/Source/NNEDenoiser/Private/NNEDenoiserPathTracingSpatialTemporalDenoiser.cpp b/Engine/Plugins/Experimental/NNEDenoiser/Source/NNEDenoiser/Private/NNEDenoiserPathTracingSpatialTemporalDenoiser.cpp
new file mode 100644
--- /dev/null
+++ b/Engine/Plugins/Experimental/NNEDenoiser/Source/NNEDenoiser/Private/NNEDenoiserPathTracingSpatialTemporalDenoiser.cpp
@@ -0,0 +1,41 @@
+// Copyright Epic Games, Inc. All Rights Reserved.
+
+#include "NNEDenoiserPathTracingSpatialTemporalDenoiser.h"
+#include "NNEDenoiserGenericDenoiser.h"
+#include "NNEDenoiserHistory.h"
+#include "NNEDenoiserLog.h"
+#include "SceneView.h"
+
+namespace UE::NNEDenoiser::Private
+{
+
+using UE::Renderer::Private::IPathTracingSpatialTemporalDenoiser;
+
+FPathTracingSpatialTemporalDenoiser::FPathTracingSpatialTemporalDenoiser(TUniquePtr<FGenericDenoiser> Denoiser) : Denoiser(MoveTemp(Denoiser))
+{
+
+}
+
+const TCHAR* FPathTracingSpatialTemporalDenoiser::GetDebugName() const
+{
+	return Denoiser->GetDebugName();
+}
+
+IPathTracingSpatialTemporalDenoiser::FOutputs FPathTracingSpatialTemporalDenoiser::AddPasses(FRDGBuilder& GraphBuilder, const FSceneView& View, const FInputs& Inputs) const
+{
+	FHistory* PrevHistory = Inputs.PrevHistory.IsValid() ? reinterpret_cast<FHistory*>(Inputs.PrevHistory.GetReference()) : nullptr;
+	
+	TUniquePtr<FHistory> NewHistory = Denoiser->AddPasses(GraphBuilder, Inputs.ColorTex, Inputs.AlbedoTex, Inputs.NormalTex, Inputs.OutputTex, Inputs.FlowTex, View.GPUMask, PrevHistory);
+
+	FOutputs Outputs;
+	Outputs.NewHistory = NewHistory.Release();
+
+	return Outputs;
+}
+
+void FPathTracingSpatialTemporalDenoiser::AddMotionVectorPass(FRDGBuilder& GraphBuilder, const FSceneView& View, const FMotionVectorInputs& Inputs) const
+{
+	unimplemented();
+}
+
+} // namespace UE::NNEDenoiser::Private
\ No newline at end of file
diff --git a/Engine/Plugins/Experimental/NNEDenoiser/Source/NNEDenoiser/Private/NNEDenoiserPathTracingSpatialTemporalDenoiser.h b/Engine/Plugins/Experimental/NNEDenoiser/Source/NNEDenoiser/Private/NNEDenoiserPathTracingSpatialTemporalDenoiser.h
new file mode 100644
--- /dev/null
+++ b/Engine/Plugins/Experimental/NNEDenoiser/Source/NNEDenoiser/Private/NNEDenoiserPathTracingSpatialTemporalDenoiser.h
@@ -0,0 +1,35 @@
+// Copyright Epic Games, Inc. All Rights Reserved.
+
+#pragma once
+
+#include "PathTracingDenoiser.h"
+#include "Templates/UniquePtr.h"
+
+namespace UE::NNEDenoiser::Private
+{
+
+class FGenericDenoiser;
+
+using UE::Renderer::Private::IPathTracingSpatialTemporalDenoiser;
+
+class FPathTracingSpatialTemporalDenoiser final : public IPathTracingSpatialTemporalDenoiser
+{
+public:
+	explicit FPathTracingSpatialTemporalDenoiser(TUniquePtr<FGenericDenoiser> Denoiser);
+
+	~FPathTracingSpatialTemporalDenoiser() = default;
+
+	// IPathTracingSpatialTemporalDenoiser interface
+	const TCHAR* GetDebugName() const override;
+
+	bool NeedTextureCreateExtraFlags() const override { return true; }
+
+	FOutputs AddPasses(FRDGBuilder& GraphBuilder, const FSceneView& View, const FInputs& Inputs) const override;
+
+	void AddMotionVectorPass(FRDGBuilder& GraphBuilder, const FSceneView& View, const FMotionVectorInputs& Inputs) const override;
+
+private:
+	TUniquePtr<FGenericDenoiser> Denoiser;
+};
+
+} // namespace UE::NNEDenoiser::Private
\ No newline at end of file
diff --git a/Engine/Plugins/Experimental/NNEDenoiser/Source/NNEDenoiser/Private/NNEDenoiserResourceManager.cpp b/Engine/Plugins/Experimental/NNEDenoiser/Source/NNEDenoiser/Private/NNEDenoiserResourceManager.cpp
new file mode 100644
--- /dev/null
+++ b/Engine/Plugins/Experimental/NNEDenoiser/Source/NNEDenoiser/Private/NNEDenoiserResourceManager.cpp
@@ -0,0 +1,191 @@
+// Copyright Epic Games, Inc. All Rights Reserved.
+
+#include "NNEDenoiserResourceManager.h"
+#include "Algo/Transform.h"
+#include "NNEDenoiserUtils.h"
+#include "NNERuntimeRDG.h"
+#include "RenderGraphBuilder.h"
+#include "RenderGraphUtils.h"
+
+namespace UE::NNEDenoiser::Private
+{
+
+void FResourceManager::AddTexture(EResourceName Name, FRDGTextureRef Current, int32 NumFrames)
+{
+	check(!TextureMap.Contains(Name));
+	check(Current != FRDGTextureRef{});
+	check(NumFrames > 0);
+
+	// Init with current texture at index 0
+	TArray<FRDGTextureRef>& Textures = TextureMap.Emplace(Name, TArray<FRDGTextureRef>{ Current });
+
+	// Find render targets in history, register them as external textures to RDG and add starting at index 1
+	if (const auto& RenderTargets = ResourceMap.Find(Name); RenderTargets)
+	{
+		Algo::Transform(*RenderTargets, Textures, [&] (TRefCountPtr<IPooledRenderTarget> RenderTarget)
+		{
+			return GraphBuilder.RegisterExternalTexture(RenderTarget);
+		});
+	}
+
+	auto FillMissingTextures = [&] (TArray<FRDGTextureRef>& Textures, int32 Num)
+	{
+		if (Textures.Num() >= Num)
+		{
+			return;
+		}
+
+		FRDGTextureRef RefTextue = Textures.Last();
+
+		const int32 OldNum = Textures.Num();
+
+		Textures.SetNumUninitialized(Num);
+		for (int32 I = OldNum; I < Textures.Num(); I++)
+		{
+			Textures[I] = RefTextue;
+		}
+	};
+
+	// Repeat "last" valid texture for frames missing their texture
+	FillMissingTextures(Textures, NumFrames);
+
+	// note: Always create intermediate texture for now.
+	const bool bNeedIntermediateTexture = true;
+
+	if (bNeedIntermediateTexture)
+	{
+		check(!IntermediateTextureMap.Contains(Name));
+
+		FRDGTextureDesc Desc = FRDGTextureDesc::Create2D(Tiling.TileSize, Current->Desc.Format, Current->Desc.ClearValue, Current->Desc.Flags);
+
+		TArray<FRDGTextureRef>& IntermediateTextures = IntermediateTextureMap.Add(Name);
+		IntermediateTextures.SetNumUninitialized(NumFrames);
+
+		for (int32 I = 0; I < IntermediateTextures.Num(); I++)
+		{
+			IntermediateTextures[I] = GraphBuilder.CreateTexture(Desc, TEXT("NNEDenoiser.IntermediateTexture"));
+		}
+	}
+}
+
+void FResourceManager::BeginTile(int32 TileIndex)
+{
+	check(TileIndex >= 0);
+	check(TileIndex < Tiling.Tiles.Num());
+
+	CurrentTileIdx = TileIndex;
+	const FTile& Tile = Tiling.Tiles[CurrentTileIdx];
+
+	// Copy sub region from original input texture to tile input texture
+	FRHICopyTextureInfo CopyInputInfo{};
+	CopyInputInfo.SourcePosition = {Tile.Position.X, Tile.Position.Y, 0};
+	CopyInputInfo.Size = {Tiling.TileSize.X, Tiling.TileSize.Y, 1};
+
+	for (const auto& KeyValue : TextureMap)
+	{
+		if (!IntermediateTextureMap.Contains(KeyValue.Key))
+		{
+			continue;
+		}
+
+		const TArray<FRDGTextureRef>& Textures = KeyValue.Value;
+		const TArray<FRDGTextureRef>& IntermediateTextures = IntermediateTextureMap.FindChecked(KeyValue.Key);
+		check(Textures.Num() == IntermediateTextures.Num());
+
+		const int32 StartIdx = KeyValue.Key == EResourceName::Output ? 1 : 0;
+		for (int32 I = StartIdx; I < Textures.Num(); I++)
+		{
+			AddCopyTexturePass(GraphBuilder, Textures[I], IntermediateTextures[I], CopyInputInfo);
+		}
+	}
+}
+
+void FResourceManager::EndTile()
+{
+	if (!IntermediateTextureMap.Contains(EResourceName::Output))
+	{
+		return;
+	}
+
+	check(CurrentTileIdx >= 0);
+	check(CurrentTileIdx < Tiling.Tiles.Num());
+	
+	const FTile& Tile = Tiling.Tiles[CurrentTileIdx];
+
+	FRHICopyTextureInfo CopyOutputInfo{};
+	CopyOutputInfo.SourcePosition = {Tile.OutputOffsets.Min.X, Tile.OutputOffsets.Min.Y, 0};
+	CopyOutputInfo.DestPosition = {Tile.Position.X + Tile.OutputOffsets.Min.X, Tile.Position.Y + Tile.OutputOffsets.Min.Y, 0};
+	CopyOutputInfo.Size = {Tiling.TileSize.X + Tile.OutputOffsets.Width(), Tiling.TileSize.Y + Tile.OutputOffsets.Height(), 1};
+
+	const TArray<FRDGTextureRef>& Textures = TextureMap.FindChecked(EResourceName::Output);
+	const TArray<FRDGTextureRef>& IntermediateTextures = IntermediateTextureMap.FindChecked(EResourceName::Output);
+	check(!Textures.IsEmpty());
+	check(Textures.Num() == IntermediateTextures.Num());
+
+	AddCopyTexturePass(GraphBuilder, IntermediateTextures[0], Textures[0], CopyOutputInfo);
+}
+
+FRDGTextureRef FResourceManager::GetIntermediateTexture(EResourceName Name, int32 FrameIdx) const
+{
+	const TArray<FRDGTextureRef>& IntermediateTextures = IntermediateTextureMap.FindChecked(Name);
+
+	check(FrameIdx >= 0);
+	check(FrameIdx < IntermediateTextures.Num());
+	return IntermediateTextures[FrameIdx];
+}
+
+TMap<EResourceName, TArray<TRefCountPtr<IPooledRenderTarget>>> FResourceManager::MakeHistoryResourceMap()
+{
+	TMap<EResourceName, TArray<TRefCountPtr<IPooledRenderTarget>>> Result;
+
+	for (const auto& KeyValue : TextureMap)
+	{
+		EResourceName Name = KeyValue.Key;
+		const TArray<FRDGTextureRef>& Textures = KeyValue.Value;
+		FRDGTextureRef CurrentTexture = Textures[0];
+
+		const int32 HistoryNumFrames = FMath::Max(Textures.Num() - 1, 0);
+
+		TArray<TRefCountPtr<IPooledRenderTarget>> NewHistoryRenderTargets;
+		if (HistoryNumFrames > 1)
+		{
+			if (const auto* RenderTargets = ResourceMap.Find(Name); RenderTargets)
+			{
+				NewHistoryRenderTargets.Append(RenderTargets->GetData(), FMath::Min(RenderTargets->Num(), HistoryNumFrames - 1));
+			}
+		}
+
+		if (HistoryNumFrames)
+		{
+			GraphBuilder.QueueTextureExtraction(CurrentTexture, &NewHistoryRenderTargets.InsertDefaulted_GetRef(0));
+		}
+
+		if (!NewHistoryRenderTargets.IsEmpty())
+		{
+			Result.Add(Name, MoveTemp(NewHistoryRenderTargets));
+		}
+	}
+
+	return Result;
+}
+
+FRDGTextureRef FResourceManager::GetTexture(EResourceName Name, int32 FrameIdx, bool bBaseTexture) const
+{
+	const bool bIntermediate = !bBaseTexture && IntermediateTextureMap.Contains(Name);
+
+	const TMap<EResourceName, TArray<FRDGTextureRef>>& Map = bIntermediate ? IntermediateTextureMap : TextureMap;
+	check(Map.Contains(Name));
+
+	const TArray<FRDGTextureRef>& Textures = Map.FindChecked(Name);
+	check(FrameIdx >= 0);
+	check(FrameIdx < Textures.Num());
+
+	if (FrameIdx < 0 || FrameIdx >= Textures.Num())
+	{
+		return FRDGTextureRef{};
+	}
+
+	return Textures[FrameIdx];
+}
+
+}
\ No newline at end of file
diff --git a/Engine/Plugins/Experimental/NNEDenoiser/Source/NNEDenoiser/Private/NNEDenoiserResourceManager.h b/Engine/Plugins/Experimental/NNEDenoiser/Source/NNEDenoiser/Private/NNEDenoiserResourceManager.h
new file mode 100644
--- /dev/null
+++ b/Engine/Plugins/Experimental/NNEDenoiser/Source/NNEDenoiser/Private/NNEDenoiserResourceManager.h
@@ -0,0 +1,55 @@
+// Copyright Epic Games, Inc. All Rights Reserved.
+
+#pragma once
+
+#include "Containers/Map.h"
+#include "NNEDenoiserModelIOMappingData.h"
+#include "NNEDenoiserTiling.h"
+#include "RendererInterface.h"
+#include "RenderGraphFwd.h"
+
+namespace UE::NNE
+{
+	struct FTensorBindingRDG;
+}
+
+namespace UE::NNEDenoiser::Private
+{
+
+class FResourceManager
+{
+public:
+	FResourceManager(FRDGBuilder& GraphBuilder, FTiling Tiling, TMap<EResourceName, TArray<TRefCountPtr<IPooledRenderTarget>>> ResourceMap = {}) :
+		GraphBuilder(GraphBuilder), Tiling(Tiling), ResourceMap(ResourceMap)
+	{
+
+	}
+
+	void AddTexture(EResourceName Name, FRDGTextureRef Current, int32 NumFrames);
+
+	void BeginTile(int32 TileIndex);
+	void EndTile();
+
+	FRDGTextureRef GetIntermediateTexture(EResourceName Name, int32 FrameIdx) const;
+
+	FRDGTextureRef GetTexture(EResourceName Name, int32 FrameIdx) const
+	{
+		return GetTexture(Name, FrameIdx, false);
+	}
+
+	TMap<EResourceName, TArray<TRefCountPtr<IPooledRenderTarget>>> MakeHistoryResourceMap();
+
+private:
+	FRDGTextureRef GetTexture(EResourceName Name, int32 FrameIdx, bool bBaseTexture) const;
+	
+	FRDGBuilder& GraphBuilder;
+	FTiling Tiling;
+	TMap<EResourceName, TArray<TRefCountPtr<IPooledRenderTarget>>> ResourceMap;
+
+	TMap<EResourceName, TArray<FRDGTextureRef>> TextureMap;
+	TMap<EResourceName, TArray<FRDGTextureRef>> IntermediateTextureMap;
+
+	int32 CurrentTileIdx = -1;
+};
+
+} // namespace UE::NNEDenoiser::Private
\ No newline at end of file
diff --git a/Engine/Plugins/Experimental/NNEDenoiser/Source/NNEDenoiser/Private/NNEDenoiserResourceMapping.h b/Engine/Plugins/Experimental/NNEDenoiser/Source/NNEDenoiser/Private/NNEDenoiserResourceMapping.h
new file mode 100644
--- /dev/null
+++ b/Engine/Plugins/Experimental/NNEDenoiser/Source/NNEDenoiser/Private/NNEDenoiserResourceMapping.h
@@ -0,0 +1,168 @@
+// Copyright Epic Games, Inc. All Rights Reserved.
+
+#pragma once
+
+#include "Containers/Array.h"
+#include "Containers/Map.h"
+#include "NNEDenoiserModelIOMappingData.h"
+
+namespace UE::NNEDenoiser::Private
+{
+	
+struct FResourceInfo
+{
+	EResourceName Name;
+	int32 Channel;
+	int32 Frame;
+};
+
+class FResourceMapping
+{
+public:
+	void Add(FResourceInfo Info)
+	{
+		ChannelMapping.Add(MoveTemp(Info));
+	}
+
+	FResourceInfo& Add_GetRef(FResourceInfo Info)
+	{
+		return ChannelMapping.Add_GetRef(MoveTemp(Info));
+	}
+
+	const FResourceInfo& GetChecked(int32 Channel) const
+	{
+		check(Channel >= 0 && Channel < ChannelMapping.Num());
+
+		return ChannelMapping[Channel];
+	}
+
+	int32 Num() const
+	{
+		return ChannelMapping.Num();
+	}
+
+	bool HasResource(EResourceName Name) const
+	{
+		for (const auto& Info: ChannelMapping)
+		{
+			if (Info.Name == Name)
+			{
+				return true;
+			}
+		}
+		return false;
+	}
+
+	int32 NumFrames(EResourceName Name) const
+	{
+		int32 MinFrameIndex = 1;
+		for (const auto& Info: ChannelMapping)
+		{
+			if (Info.Name == Name)
+			{
+				MinFrameIndex = FMath::Min(MinFrameIndex, Info.Frame);
+			}
+		}
+
+		return 1 - MinFrameIndex;
+	}
+
+	TMap<int32, TArray<FIntPoint>> GetChannelMappingPerFrame(EResourceName Name) const
+	{
+		TMap<int32, TArray<FIntPoint>> Result;
+		for (int32 I = 0; I < ChannelMapping.Num(); I++)
+		{
+			const FResourceInfo& Info = ChannelMapping[I];
+
+			if (Info.Name == Name)
+			{
+				Result.FindOrAdd(Info.Frame).Emplace(I, Info.Channel);
+			}
+		}
+
+		return Result;
+	}
+
+private:
+	TArray<FResourceInfo> ChannelMapping;
+};
+
+class FResourceMappingList
+{
+public:
+	void Add(FResourceMapping ResourceMapping)
+	{
+		InputMapping.Add(MoveTemp(ResourceMapping));
+	}
+
+	FResourceMapping& Add_GetRef(FResourceMapping ResourceMapping)
+	{
+		return InputMapping.Add_GetRef(MoveTemp(ResourceMapping));
+	}
+
+	const FResourceMapping& GetChecked(int32 InputIndex) const
+	{
+		check(InputIndex >= 0 && InputIndex < InputMapping.Num());
+
+		return InputMapping[InputIndex];
+	}
+
+	int32 Num() const
+	{
+		return InputMapping.Num();
+	}
+
+	int32 NumChannels(int32 InputIndex) const
+	{
+		check(InputIndex >= 0 && InputIndex < InputMapping.Num());
+		
+		return InputMapping[InputIndex].Num();
+	}
+
+	int32 NumFrames(EResourceName Name) const
+	{
+		int32 MaxNumFrames = 0;
+		for (const auto& Input : InputMapping)
+		{
+			const int32 NumFrames = Input.NumFrames(Name);
+			if (MaxNumFrames < NumFrames)
+			{
+				MaxNumFrames = NumFrames;
+			}
+		}
+
+		return MaxNumFrames;
+	}
+
+	bool HasResource(EResourceName Name) const
+	{
+		for (const auto& Input : InputMapping)
+		{
+			if (Input.HasResource(Name))
+			{
+				return true;
+			}
+		}
+
+		return false;
+	}
+
+	bool HasResource(int32 InputIndex, EResourceName Name) const
+	{
+		check(InputIndex >= 0 && InputIndex < InputMapping.Num());
+
+		return InputMapping[InputIndex].HasResource(Name);
+	}
+
+	TMap<int32, TArray<FIntPoint>> GetChannelMappingPerFrame(int32 InputIndex, EResourceName Name) const
+	{
+		check(InputIndex >= 0 && InputIndex < InputMapping.Num());
+
+		return InputMapping[InputIndex].GetChannelMappingPerFrame(Name);
+	}
+
+private:
+	TArray<FResourceMapping> InputMapping;
+};
+
+} // namespace UE::NNEDenoiser::Private
\ No newline at end of file
diff --git a/Engine/Plugins/Experimental/NNEDenoiser/Source/NNEDenoiser/Private/NNEDenoiserSettings.cpp b/Engine/Plugins/Experimental/NNEDenoiser/Source/NNEDenoiser/Private/NNEDenoiserSettings.cpp
new file mode 100644
--- /dev/null
+++ b/Engine/Plugins/Experimental/NNEDenoiser/Source/NNEDenoiser/Private/NNEDenoiserSettings.cpp
@@ -0,0 +1,26 @@
+// Copyright Epic Games, Inc. All Rights Reserved.
+
+#include "NNEDenoiserSettings.h"
+#include "HAL/IConsoleManager.h"
+#include "Misc/ConfigUtilities.h"
+
+#include UE_INLINE_GENERATED_CPP_BY_NAME(NNEDenoiserSettings)
+
+UNNEDenoiserSettings::UNNEDenoiserSettings()
+{
+	CategoryName = TEXT("Plugins");
+	SectionName = TEXT("NNE Denoiser");
+
+	DenoiserModelData = FSoftObjectPath(TEXT("/NNEDenoiser/NNEDNN_Oidn2_Color_720.NNEDNN_Oidn2_Color_720"));
+}
+
+void UNNEDenoiserSettings::PostInitProperties()
+{
+	if (IsTemplate())
+	{
+		// We want the .ini file to have precedence over the CVar constructor, so we apply the ini to the CVar before following the regular UDeveloperSettingsBackedByCVars flow
+		UE::ConfigUtilities::ApplyCVarSettingsFromIni(TEXT("/Script/NNEDenoiser.NNEDenoiserSettings"), *GEngineIni, ECVF_SetByProjectSetting);
+	}
+
+	Super::PostInitProperties();
+}
\ No newline at end of file
diff --git a/Engine/Plugins/Experimental/NNEDenoiser/Source/NNEDenoiser/Private/NNEDenoiserTiling.h b/Engine/Plugins/Experimental/NNEDenoiser/Source/NNEDenoiser/Private/NNEDenoiserTiling.h
new file mode 100644
--- /dev/null
+++ b/Engine/Plugins/Experimental/NNEDenoiser/Source/NNEDenoiser/Private/NNEDenoiserTiling.h
@@ -0,0 +1,95 @@
+// Copyright Epic Games, Inc. All Rights Reserved.
+
+#pragma once
+
+#include "Math/IntPoint.h"
+#include "Math/IntRect.h"
+#include "Math/IntVector.h"
+#include "Math/MathFwd.h"
+
+namespace UE::NNEDenoiser::Private
+{
+	
+	struct FTile
+	{
+		FIntPoint Position;
+		FIntRect OutputOffsets;
+	};
+
+	struct FTiling
+	{
+		FIntPoint TileSize;
+		FIntPoint Count;
+		TArray<FTile> Tiles;
+	};
+
+	inline FTiling CreateTiling(FIntPoint TargetTileSize, FIntPoint MinimumOverlap, FIntPoint Size)
+	{
+		auto GetNumTiles = [] (int32 TileSize, int32 Overlap, int32 Size)
+		{
+			return FMath::Max(FMath::CeilToInt32((Size - TileSize) / (float)(TileSize - Overlap)), 0) + 1;
+		};
+
+		auto GetOffsets = [] (int32 Count, int32 TileSize, int32 Overlap, int32 Size)
+		{
+			TArray<int32> Result;
+			for (int32 I = 0; I < Count; I++)
+			{
+				Result.Add(I < Count - 1 ? I * (TileSize - Overlap) : Size - TileSize);
+			}
+
+			return Result;
+		};
+
+		FTiling Result{};
+		Result.TileSize = {
+			FMath::Min(TargetTileSize.X, Size.X),
+			FMath::Min(TargetTileSize.Y, Size.Y)
+		};
+
+		Result.Count = {
+			GetNumTiles(Result.TileSize.X, MinimumOverlap.X, Size.X),
+			GetNumTiles(Result.TileSize.Y, MinimumOverlap.Y, Size.Y)
+		};
+
+		const FIntPoint TotalOverlap = Result.Count * Result.TileSize - Size;
+		const FIntPoint Overlap = {
+			Result.Count.X == 1 ? 0 : TotalOverlap.X / (Result.Count.X - 1),
+			Result.Count.Y == 1 ? 0 : TotalOverlap.Y / (Result.Count.Y - 1)
+		};
+
+		const FIntPoint HalfOverlap = {
+			FMath::FloorToInt32(Overlap.X / 2.0f),
+			FMath::FloorToInt32(Overlap.Y / 2.0f)
+		};
+
+		const TArray<int32> OffsetsX = GetOffsets(Result.Count.X, Result.TileSize.X, Overlap.X, Size.X);
+		const TArray<int32> OffsetsY = GetOffsets(Result.Count.Y, Result.TileSize.Y, Overlap.Y, Size.Y);
+
+		for (int32 Ty = 0; Ty < Result.Count.Y; Ty++)
+		{
+			const int32 Y0 = OffsetsY[Ty];
+			const int32 Y1 = Y0 + Result.TileSize.Y;
+
+			const int32 OutY0 = Ty > 0 ? HalfOverlap.Y : 0;
+			const int32 OutY1 = Ty < Result.Count.Y - 1 ? -HalfOverlap.Y : 0;
+
+			for (int32 Tx = 0; Tx < Result.Count.X; Tx++)
+			{
+				const int32 X0 = OffsetsX[Tx];
+				const int32 X1 = X0 + Result.TileSize.X;
+
+				const int32 OutX0 = Tx > 0 ? HalfOverlap.X : 0;
+				const int32 OutX1 = Tx < Result.Count.X - 1 ? -HalfOverlap.X : 0;
+
+				const FIntPoint InputPosition = {X0, Y0};
+				const FIntRect OutputOffsets = {OutX0, OutY0, OutX1, OutY1};
+
+				Result.Tiles.Add(FTile{InputPosition, OutputOffsets});
+			}
+		}
+
+		return Result;
+	}
+
+}
\ No newline at end of file
diff --git a/Engine/Plugins/Experimental/NNEDenoiser/Source/NNEDenoiser/Private/NNEDenoiserTransferFunction.h b/Engine/Plugins/Experimental/NNEDenoiser/Source/NNEDenoiser/Private/NNEDenoiserTransferFunction.h
new file mode 100644
--- /dev/null
+++ b/Engine/Plugins/Experimental/NNEDenoiser/Source/NNEDenoiser/Private/NNEDenoiserTransferFunction.h
@@ -0,0 +1,88 @@
+// Copyright Epic Games, Inc. All Rights Reserved.
+
+#pragma once
+
+#include "Math/UnrealMathUtility.h"
+
+namespace UE::NNEDenoiser::Private
+{
+	constexpr float HDRMax = 65504.f; // maximum HDR value
+
+	// from [REDACTED_PROFILE_URL]/oidn
+	struct FPU
+	{
+		static constexpr float a  =  1.41283765e+03f;
+		static constexpr float b  =  1.64593172e+00f;
+		static constexpr float c  =  4.31384981e-01f;
+		static constexpr float d  = -2.94139609e-03f;
+		static constexpr float e  =  1.92653254e-01f;
+		static constexpr float f  =  6.26026094e-03f;
+		static constexpr float g  =  9.98620152e-01f;
+		static constexpr float y0 =  1.57945760e-06f;
+		static constexpr float y1 =  3.22087631e-02f;
+		static constexpr float x0 =  2.23151711e-03f;
+		static constexpr float x1 =  3.70974749e-01f;
+
+		static float Forward(float y)
+		{
+			if (y <= y0)
+			return a * y;
+			else if (y <= y1)
+			return b * FMath::Pow(y, c) + d;
+			else
+			return e * FMath::Loge(y + f) + g;
+		}
+
+		static float Inverse(float x)
+		{
+			if (x <= x0)
+			return x / a;
+			else if (x <= x1)
+			return FMath::Pow((x - d) / b, 1.f/c);
+			else
+			return FMath::Exp((x - g) / e) - f;
+		}
+	};
+
+	class ITransferFunction
+	{
+		public:
+			virtual float Forward(float Y) = 0;
+			virtual float Inverse(float X) = 0;
+	};
+
+	class FHDRTransferFunction : public ITransferFunction
+	{
+		public:
+			explicit FHDRTransferFunction(float Ymax)
+			{
+				InvNormScale = FPU::Forward(Ymax);
+				NormScale = 1.0f / InvNormScale;
+			}
+
+			float Forward(float Y) override
+			{
+				return FPU::Forward(Y) * NormScale;
+			}
+
+			float Inverse(float X) override
+			{
+				return FPU::Inverse(X * InvNormScale);
+			}
+
+			float GetNormScale() const
+			{
+				return NormScale;
+			}
+
+			float GetInvNormScale() const
+			{
+				return InvNormScale;
+			}
+
+		private:
+			float NormScale = 1.0f;
+			float InvNormScale = 1.0f;
+	};
+
+}
\ No newline at end of file
diff --git a/Engine/Plugins/Experimental/NNEDenoiser/Source/NNEDenoiser/Private/NNEDenoiserUtils.h b/Engine/Plugins/Experimental/NNEDenoiser/Source/NNEDenoiser/Private/NNEDenoiserUtils.h
new file mode 100644
--- /dev/null
+++ b/Engine/Plugins/Experimental/NNEDenoiser/Source/NNEDenoiser/Private/NNEDenoiserUtils.h
@@ -0,0 +1,140 @@
+// Copyright Epic Games, Inc. All Rights Reserved.
+
+#pragma once
+
+#include "NNEDenoiserLog.h"
+#include "NNEDenoiserModelIOMappingData.h"
+#include "NNEDenoiserParameters.h"
+#include "NNEDenoiserShadersDefaultCS.h"
+#include "NNEDenoiserShadersOidnCS.h"
+#include "NNEDenoiserTiling.h"
+#include "NNETypes.h"
+#include "RHICommandList.h"
+#include "RHIResources.h"
+#include "RHITypes.h"
+
+namespace UE::NNEDenoiser::Private
+{
+
+	template<class IntType>
+	bool IsTensorShapeValid(TConstArrayView<IntType> ShapeData, TConstArrayView<int32> RequiredShapeData, const FString& Label)
+	{
+		static_assert(std::is_integral_v<IntType>);
+
+		if (ShapeData.Num() != RequiredShapeData.Num())
+		{
+			UE_LOG(LogNNEDenoiser, Error, TEXT("%s has wrong rank (expected %d, got %d)!"), *Label, RequiredShapeData.Num(), ShapeData.Num())
+			return false;
+		}
+
+		for (int32 I = 0; I < RequiredShapeData.Num(); I++)
+		{
+			if (RequiredShapeData[I] >= 0 && (int32)ShapeData[I] != RequiredShapeData[I])
+			{
+				UE_LOG(LogNNEDenoiser, Error, TEXT("%s does not have required shape!"), *Label)
+				return false;
+			}
+		}
+
+		return true;
+	}
+
+	template <typename PixelType>
+	void CopyTextureFromGPUToCPU(FRHICommandListImmediate& RHICmdList, FRHITexture* SrcTexture, FIntPoint Size, TArray<PixelType>& DstArray)
+	{
+		SCOPED_NAMED_EVENT_TEXT("NNEDenoiser.CopyTextureFromGPUToCPU", FColor::Magenta);
+
+		uint32 SrcStride = 0;
+		const PixelType* SrcBuffer = static_cast<PixelType*>(RHICmdList.LockTexture2D(SrcTexture, 0, RLM_ReadOnly, SrcStride, false));
+		SrcStride /= sizeof(PixelType);
+		PixelType* DstBuffer = DstArray.GetData();
+		for (int32 Y = 0; Y < Size.Y; Y++, DstBuffer += Size.X, SrcBuffer += SrcStride)
+		{
+			FPlatformMemory::Memcpy(DstBuffer, SrcBuffer, Size.X * sizeof(PixelType));
+
+		}
+		RHICmdList.UnlockTexture2D(SrcTexture, 0, false);
+	}
+
+	template <typename PixelType>
+	void CopyTextureFromCPUToGPU(FRHICommandListImmediate& RHICmdList, const TArray<PixelType>& SrcArray, FIntPoint Size, FRHITexture* DstTexture)
+	{
+		SCOPED_NAMED_EVENT_TEXT("NNEDenoiser.CopyTextureFromCPUToGPU", FColor::Magenta);
+
+		uint32 DestStride;
+		FLinearColor* DstBuffer = static_cast<PixelType*>(RHICmdList.LockTexture2D(DstTexture, 0, RLM_WriteOnly, DestStride, false));
+		DestStride /= sizeof(PixelType);
+		const FLinearColor* SrcBuffer = SrcArray.GetData();
+		for (int32 Y = 0; Y < Size.Y; Y++, SrcBuffer += Size.X, DstBuffer += DestStride)
+		{
+			FPlatformMemory::Memcpy(DstBuffer, SrcBuffer, Size.X * sizeof(PixelType));
+		}
+		RHICmdList.UnlockTexture2D(DstTexture, 0, false);
+	}
+
+	inline void CopyBufferFromGPUToCPU(FRHICommandListImmediate& RHICmdList, FRHIBuffer* SrcBuffer, int32 Count, TArray<uint8>& DstArray)
+	{
+		SCOPED_NAMED_EVENT_TEXT("NNEDenoiser.CopyBufferFromGPUToCPU", FColor::Magenta);
+
+		const uint8* Src = static_cast<uint8*>(RHICmdList.LockBuffer(SrcBuffer, 0, Count, RLM_ReadOnly));
+		uint8* Dst = DstArray.GetData();
+		
+		FPlatformMemory::Memcpy(Dst, Src, Count);
+
+		RHICmdList.UnlockBuffer(SrcBuffer);
+	}
+
+	inline void CopyBufferFromCPUToGPU(FRHICommandListImmediate& RHICmdList, const TArray<uint8>& SrcArray, int32 Count, FRHIBuffer* DstBuffer)
+	{
+		SCOPED_NAMED_EVENT_TEXT("NNEDenoiser.CopyBufferFromCPUToGPU", FColor::Magenta);
+
+		uint8* Dst = static_cast<uint8*>(RHICmdList.LockBuffer(DstBuffer, 0, Count, RLM_WriteOnly));
+		const uint8* Src = SrcArray.GetData();
+		
+		FPlatformMemory::Memcpy(Dst, Src, Count);
+
+		RHICmdList.UnlockBuffer(DstBuffer);
+	}
+
+	inline EPixelFormat GetBufferFormat(ENNETensorDataType TensorDataType)
+	{
+		switch (TensorDataType)
+		{
+			case ENNETensorDataType::Half: return EPixelFormat::PF_R16F;
+			case ENNETensorDataType::Float: return EPixelFormat::PF_R32_FLOAT;
+		}
+		return EPixelFormat::PF_Unknown;
+	}
+
+	inline UE::NNEDenoiserShaders::Internal::ENNEDenoiserDataType GetDenoiserShaderDataType(ENNETensorDataType TensorDataType)
+	{
+		using UE::NNEDenoiserShaders::Internal::ENNEDenoiserDataType;
+
+		switch (TensorDataType)
+		{
+			case ENNETensorDataType::Half: return ENNEDenoiserDataType::Half;
+			case ENNETensorDataType::Float: return ENNEDenoiserDataType::Float;
+		}
+		return ENNEDenoiserDataType::None;
+	}
+
+	inline NNEDenoiserShaders::Internal::ENNEDenoiserInputKind GetInputKind(EResourceName TensorName)
+	{
+		using NNEDenoiserShaders::Internal::ENNEDenoiserInputKind;
+		
+		switch(TensorName)
+		{
+			case EResourceName::Color:	return ENNEDenoiserInputKind::Color;
+			case EResourceName::Albedo:	return ENNEDenoiserInputKind::Albedo;
+			case EResourceName::Normal:	return ENNEDenoiserInputKind::Normal;
+			case EResourceName::Flow:	return ENNEDenoiserInputKind::Flow;
+			case EResourceName::Output:	return ENNEDenoiserInputKind::Output;
+		}
+
+		// There should be a case for every resource name
+		checkNoEntry();
+		
+		return ENNEDenoiserInputKind::Color;
+	}
+
+}
\ No newline at end of file
diff --git a/Engine/Plugins/Experimental/NNEDenoiser/Source/NNEDenoiser/Private/NNEDenoiserViewExtension.cpp b/Engine/Plugins/Experimental/NNEDenoiser/Source/NNEDenoiser/Private/NNEDenoiserViewExtension.cpp
new file mode 100644
--- /dev/null
+++ b/Engine/Plugins/Experimental/NNEDenoiser/Source/NNEDenoiser/Private/NNEDenoiserViewExtension.cpp
@@ -0,0 +1,323 @@
+// Copyright Epic Games, Inc. All Rights Reserved.
+
+#include "NNEDenoiserViewExtension.h"
+#include "NNE.h"
+#include "NNEDenoiserGenericDenoiser.h"
+#include "NNEDenoiserIOProcessBase.h"
+#include "NNEDenoiserIOProcessOidn.h"
+#include "NNEDenoiserLog.h"
+#include "NNEDenoiserModelData.h"
+#include "NNEDenoiserModelInstanceCPU.h"
+#include "NNEDenoiserModelInstanceGPU.h"
+#include "NNEDenoiserModelInstanceRDG.h"
+#include "NNEDenoiserModelIOMappingData.h"
+#include "NNEDenoiserPathTracingDenoiser.h"
+#include "NNEDenoiserPathTracingSpatialTemporalDenoiser.h"
+#include "NNEModelData.h"
+#include "PathTracingDenoiser.h"
+
+namespace UE::NNEDenoiser::Private
+{
+
+static TAutoConsoleVariable<bool> CVarNNEDenoiser(
+	TEXT("NNEDenoiser"),
+	true,
+	TEXT("Enable the NNE Denoiser.")
+);
+
+static TAutoConsoleVariable<int32> CVarNNEDenoiserModelData(
+	TEXT("NNEDenoiser.ModelData"),
+	0,
+	TEXT("Defines the model data of the denoiser model.\n")
+	TEXT("  0: Use model data from Project Settings\n")
+	TEXT("  1: OIDN Color\n")
+	TEXT("  2: OIDN Color and Albedo\n")
+	TEXT("  3: OIDN Color, Albedo and Normal\n")
+	TEXT("  4: OIDN Color, Half\n")
+	TEXT("  5: OIDN Color and Albedo, Half\n")
+	TEXT("  6: OIDN Color, Albedo and Normal, Half")
+);
+
+static TAutoConsoleVariable<int32> CVarNNEDenoiserRuntimeType(
+	TEXT("NNEDenoiser.Runtime.Type"),
+	0,
+	TEXT("Defines the runtime type to run the denoiser model.\n")
+	TEXT("  0: CPU\n")
+	TEXT("  1: GPU\n")
+	TEXT("  2: RDG\n")
+);
+
+static TAutoConsoleVariable<FString> CVarNNEDenoiserRuntimeName(
+	TEXT("NNEDenoiser.Runtime.Name"),
+	FString(),
+	TEXT("Defines the runtime name to run the denoiser model. Leave empty to use default.")
+);
+
+EDenoiserRuntimeType GetDenoiserRuntimeTypeFromCVar()
+{
+	const int32 Value = CVarNNEDenoiserRuntimeType.GetValueOnGameThread();
+	const int32 Min = (int32)EDenoiserRuntimeType::CPU;
+	const int32 Max = (int32)EDenoiserRuntimeType::RDG;
+	return static_cast<EDenoiserRuntimeType>(FMath::Clamp(Value, Min, Max));
+}
+
+FString GetDenoiserModelDataNameFromCVarAndSettings(const UNNEDenoiserSettings* Settings)
+{
+	const int32 Idx = CVarNNEDenoiserModelData.GetValueOnGameThread();
+	switch(Idx)
+	{
+		case 0: return !Settings->DenoiserModelData.IsNull() ? Settings->DenoiserModelData.ToString() : FString();
+		case 1: return TEXT("/NNEDenoiser/NNEDNN_Oidn2_Color_720.NNEDNN_Oidn2_Color_720");
+		case 2: return TEXT("/NNEDenoiser/NNEDNN_Oidn2_ColorAlbedo_720.NNEDNN_Oidn2_ColorAlbedo_720");
+		case 3: return TEXT("/NNEDenoiser/NNEDNN_Oidn2_ColorAlbedoNormal_720.NNEDNN_Oidn2_ColorAlbedoNormal_720");
+		case 4: return TEXT("/NNEDenoiser/NNEDNN_Oidn2_Color_720_Float16.NNEDNN_Oidn2_Color_720_Float16");
+		case 5: return TEXT("/NNEDenoiser/NNEDNN_Oidn2_ColorAlbedo_720_Float16.NNEDNN_Oidn2_ColorAlbedo_720_Float16");
+		case 6: return TEXT("/NNEDenoiser/NNEDNN_Oidn2_ColorAlbedoNormal_720_Float16.NNEDNN_Oidn2_ColorAlbedoNormal_720_Float16");
+	}
+	check(false);
+	return FString();
+}
+
+TUniquePtr<FGenericDenoiser> CreateNNEDenoiser(UNNEModelData& ModelData, EDenoiserRuntimeType RuntimeType, const FString& RuntimeNameOverride,
+	TUniquePtr<IInputProcess> InputProcess, TUniquePtr<IOutputProcess> OutputProcess, FParameters Parameters)
+{
+	TUniquePtr<IModelInstance> ModelInstance;
+	if (RuntimeType == EDenoiserRuntimeType::CPU)
+	{
+		ModelInstance = FModelInstanceCPU::Make(ModelData, RuntimeNameOverride);
+	}
+	else if (RuntimeType == EDenoiserRuntimeType::GPU)
+	{
+		ModelInstance = FModelInstanceGPU::Make(ModelData, RuntimeNameOverride);
+	}
+	else if (RuntimeType == EDenoiserRuntimeType::RDG)
+	{
+		ModelInstance = FModelInstanceRDG::Make(ModelData, RuntimeNameOverride);
+	}
+
+	if (!ModelInstance.IsValid())
+	{
+		UE_LOG(LogNNEDenoiser, Error, TEXT("Could not create denoiser!"));
+		return {};
+	}
+
+	return MakeUnique<FGenericDenoiser>(MoveTemp(ModelInstance), MoveTemp(InputProcess), MoveTemp(OutputProcess), MoveTemp(Parameters));
+}
+
+static FResourceMappingList MakeTensorLayout(UDataTable* DataTable)
+{
+	TMap<int32, TMap<int32, FResourceInfo>> Map;
+	DataTable->ForeachRow<FNNEDenoiserModelIOMappingData>("FResourceLayout", [&] (const FName &Key, const FNNEDenoiserModelIOMappingData &Value)
+	{
+		if (Value.TensorChannel < 0)
+		{
+			for (int32 I = 0; I < -Value.TensorChannel; I++)
+			{
+				Map.FindOrAdd(Value.TensorIndex).FindOrAdd(I) = FResourceInfo{Value.Resource, I, Value.FrameIndex};
+			}
+		}
+		else
+		{
+			Map.FindOrAdd(Value.TensorIndex).FindOrAdd(Value.TensorChannel) = FResourceInfo{Value.Resource, Value.ResourceChannel, Value.FrameIndex};
+		}
+	});
+
+	FResourceMappingList Result{};
+
+	for (int32 I = 0; I < Map.Num(); I++)
+	{
+		checkf(Map.Contains(I), TEXT("Missing intput/output %d, must be continuous!"), I);
+
+		const TMap<int32, FResourceInfo>& InnerMap = Map[I];
+		FResourceMapping& Mapping = Result.Add_GetRef({});
+
+		for (int32 J = 0; J < InnerMap.Num(); J++)
+		{
+			checkf(InnerMap.Contains(J), TEXT("Missing tensor info for channel %d, must be continuous!"), J);
+
+			Mapping.Add(InnerMap[J]);
+		}
+	}
+
+	return Result;
+}
+
+TUniquePtr<FGenericDenoiser> CreateNNEDenoiserFromAsset(const FString& AssetName, EDenoiserRuntimeType RuntimeType, const FString& RuntimeNameOverride)
+{
+	if (AssetName.IsEmpty())
+	{
+		UE_LOG(LogNNEDenoiser, Error, TEXT("Asset name not set!"));
+		return {};
+	}
+
+	UNNEDenoiserModelData *DenoiserModelData = LoadObject<UNNEDenoiserModelData>(nullptr, *AssetName);
+	if (!DenoiserModelData)
+	{
+		UE_LOG(LogNNEDenoiser, Error, TEXT("Could not load denoiser model data asset!"));
+		return {};
+	}
+
+	UNNEModelData* ModelData = DenoiserModelData->ModelData.LoadSynchronous();
+	if (!ModelData)
+	{
+		UE_LOG(LogNNEDenoiser, Error, TEXT("Asset does not contain model data!"));
+		return {};
+	}
+	UDataTable* InputMappingTable = DenoiserModelData->InputMapping.LoadSynchronous();
+	if (InputMappingTable)
+	{
+		UE_LOG(LogNNEDenoiser, Log, TEXT("Loaded input mapping from %s"), *DenoiserModelData->InputMapping.GetAssetName());
+	}
+	UDataTable* OutputMappingTable = DenoiserModelData->OutputMapping.LoadSynchronous();
+	if (OutputMappingTable)
+	{
+		UE_LOG(LogNNEDenoiser, Log, TEXT("Loaded output mapping from %s"), *DenoiserModelData->OutputMapping.GetAssetName());
+	}
+
+	FResourceMappingList InputLayout;
+	if (InputMappingTable)
+	{
+		InputLayout = MakeTensorLayout(InputMappingTable);
+	}
+
+	FResourceMappingList OutputLayout;
+	if (OutputMappingTable)
+	{
+		OutputLayout = MakeTensorLayout(OutputMappingTable);
+	}
+
+	const bool bIsOidnModel = AssetName.Contains(TEXT("oidn2"));
+	const bool bIsInHouseModel = AssetName.Contains(TEXT("kpcn"));
+
+	check(!bIsOidnModel || !bIsInHouseModel);
+
+	TUniquePtr<IInputProcess> InputProcess;
+	if (bIsOidnModel)
+	{
+		InputProcess = MakeUnique<FInputProcessOidn>(MoveTemp(InputLayout));
+	}
+	else
+	{
+		InputProcess = MakeUnique<FInputProcessBase>(MoveTemp(InputLayout));
+	}
+
+	TUniquePtr<IOutputProcess> OutputProcess;
+	if (bIsOidnModel)
+	{
+		OutputProcess = MakeUnique<FOutputProcessOidn>(MoveTemp(OutputLayout));
+	}
+	else
+	{
+		OutputProcess = MakeUnique<FOutputProcessBase>(MoveTemp(OutputLayout));
+	}
+
+	return CreateNNEDenoiser(*ModelData, RuntimeType, RuntimeNameOverride, MoveTemp(InputProcess), MoveTemp(OutputProcess), {});
+}
+
+FViewExtension::FViewExtension(const FAutoRegister& AutoRegister) : FSceneViewExtensionBase(AutoRegister)
+{
+	ApplySettings(GetDefault<UNNEDenoiserSettings>());
+
+#if WITH_EDITOR
+	GetMutableDefault<UNNEDenoiserSettings>()->OnSettingChanged().AddRaw(this, &FViewExtension::OnDenoiserSettingsChanged);
+#endif
+}
+
+FViewExtension::~FViewExtension()
+{
+#if WITH_EDITOR
+	if(UObjectInitialized())
+	{
+		GetMutableDefault<UNNEDenoiserSettings>()->OnSettingChanged().RemoveAll(this);
+	}
+#endif
+
+	GPathTracingDenoiserPlugin.Reset();
+	GPathTracingSpatialTemporalDenoiserPlugin.Reset();
+}
+
+void FViewExtension::BeginRenderViewFamily(FSceneViewFamily& InViewFamily)
+{
+	ApplySettings(GetDefault<UNNEDenoiserSettings>());
+}
+
+void FViewExtension::PreRenderViewFamily_RenderThread(FRDGBuilder& GraphBuilder, FSceneViewFamily& InViewFamily)
+{
+	if (!bDenoiserEnabled)
+	{
+		GPathTracingDenoiserPlugin.Reset();
+		GPathTracingSpatialTemporalDenoiserPlugin.Reset();
+	}
+	else
+	{
+		if (DenoiserToSwap.IsValid())
+		{
+			GPathTracingSpatialTemporalDenoiserPlugin.Reset();
+			GPathTracingDenoiserPlugin = MoveTemp(DenoiserToSwap);
+		}
+
+		if (SpatialTemporalDenoiserToSwap.IsValid())
+		{
+			GPathTracingDenoiserPlugin.Reset();
+			GPathTracingSpatialTemporalDenoiserPlugin = MoveTemp(SpatialTemporalDenoiserToSwap);
+		}
+	}
+}
+
+void FViewExtension::ApplySettings(const UNNEDenoiserSettings* Settings)
+{
+	// Check for changes
+	bool bNeedsUpdate = false;
+	bNeedsUpdate |= CVarNNEDenoiser.GetValueOnGameThread() != bDenoiserEnabled;
+	bNeedsUpdate |= GetDenoiserRuntimeTypeFromCVar() != RuntimeType;
+	bNeedsUpdate |= CVarNNEDenoiserRuntimeName.GetValueOnGameThread() != RuntimeName;
+	bNeedsUpdate |= GetDenoiserModelDataNameFromCVarAndSettings(Settings) != ModelDataName;
+
+	if (!bNeedsUpdate)
+	{
+		return;
+	}
+
+	bDenoiserEnabled = CVarNNEDenoiser.GetValueOnGameThread();
+
+	UE_LOG(LogNNEDenoiser, Log, TEXT("ApplySettings: bDenoiserEnabled %d"), bDenoiserEnabled);
+
+	if (!bDenoiserEnabled)
+	{
+		return;
+	}
+
+	RuntimeType = GetDenoiserRuntimeTypeFromCVar();
+	RuntimeName = CVarNNEDenoiserRuntimeName.GetValueOnGameThread();
+	ModelDataName = GetDenoiserModelDataNameFromCVarAndSettings(Settings);
+
+	UE_LOG(LogNNEDenoiser, Log, TEXT("Create denoiser from asset %s..."), *ModelDataName);
+
+	TUniquePtr<FGenericDenoiser> Denoiser = CreateNNEDenoiserFromAsset(ModelDataName, RuntimeType, RuntimeName);
+	if (Denoiser.IsValid())
+	{
+		if (ModelDataName.Contains(TEXT("kpcn")))
+		{
+			SpatialTemporalDenoiserToSwap = MakeUnique<FPathTracingSpatialTemporalDenoiser>(MoveTemp(Denoiser));
+		}
+		else
+		{
+			DenoiserToSwap = MakeUnique<FPathTracingDenoiser>(MoveTemp(Denoiser));
+		}
+	}
+	else
+	{
+		UE_LOG(LogNNEDenoiser, Error, TEXT("Could not create denoiser!"));
+	}
+}
+
+#if WITH_EDITOR
+void FViewExtension::OnDenoiserSettingsChanged(UObject* InObject, struct FPropertyChangedEvent& InPropertyChangedEvent)
+{
+	UE_LOG(LogNNEDenoiser, Log, TEXT("Settings %s changed: %s"), *InObject->GetName(), *InPropertyChangedEvent.GetPropertyName().ToString());
+
+	ApplySettings(CastChecked<UNNEDenoiserSettings>(InObject));
+}
+#endif
+
+} // namespace UE::NNEDenoiser::Private
\ No newline at end of file
diff --git a/Engine/Plugins/Experimental/NNEDenoiser/Source/NNEDenoiser/Private/NNEDenoiserViewExtension.h b/Engine/Plugins/Experimental/NNEDenoiser/Source/NNEDenoiser/Private/NNEDenoiserViewExtension.h
new file mode 100644
--- /dev/null
+++ b/Engine/Plugins/Experimental/NNEDenoiser/Source/NNEDenoiser/Private/NNEDenoiserViewExtension.h
@@ -0,0 +1,47 @@
+// Copyright Epic Games, Inc. All Rights Reserved.
+
+#pragma once
+
+#include "SceneViewExtension.h"
+#include "NNEDenoiserSettings.h"
+
+namespace UE::Renderer::Private
+{
+	class IPathTracingDenoiser;
+	class IPathTracingSpatialTemporalDenoiser;
+}
+
+namespace UE::NNEDenoiser::Private
+{
+
+class FViewExtension final : public FSceneViewExtensionBase
+{
+	public:
+		FViewExtension(const FAutoRegister& AutoRegister);
+		virtual ~FViewExtension();
+
+		void SetupViewFamily(FSceneViewFamily& InViewFamily) override { }
+		void SetupView(FSceneViewFamily& InViewFamily, FSceneView& InView) override { }
+			
+		void BeginRenderViewFamily(FSceneViewFamily& InViewFamily) override;
+
+		void PreRenderViewFamily_RenderThread(FRDGBuilder& GraphBuilder, FSceneViewFamily& InViewFamily) override;
+
+	private:
+		void ApplySettings(const UNNEDenoiserSettings* Settings);
+
+#if WITH_EDITOR
+		void OnDenoiserSettingsChanged(UObject* InObject, struct FPropertyChangedEvent& InPropertyChangedEvent);
+#endif
+
+		TUniquePtr<UE::Renderer::Private::IPathTracingDenoiser> DenoiserToSwap;
+		TUniquePtr<UE::Renderer::Private::IPathTracingSpatialTemporalDenoiser> SpatialTemporalDenoiserToSwap;
+
+		// Cached settings and CVar values
+		bool bDenoiserEnabled = true;
+		EDenoiserRuntimeType RuntimeType;
+		FString RuntimeName;
+		FString ModelDataName;
+	};
+
+} // namespace UE::NNEDenoiser::Private
\ No newline at end of file
diff --git a/Engine/Plugins/Experimental/NNEDenoiser/Source/NNEDenoiser/Public/NNEDenoiserModelData.h b/Engine/Plugins/Experimental/NNEDenoiser/Source/NNEDenoiser/Public/NNEDenoiserModelData.h
new file mode 100644
--- /dev/null
+++ b/Engine/Plugins/Experimental/NNEDenoiser/Source/NNEDenoiser/Public/NNEDenoiserModelData.h
@@ -0,0 +1,29 @@
+// Copyright Epic Games, Inc. All Rights Reserved.
+
+#pragma once
+
+#include "Engine/DataAsset.h"
+#include "Engine/DataTable.h"
+#include "NNEModelData.h"
+
+#include "NNEDenoiserModelData.generated.h"
+
+/** Denoiser model data asset */
+UCLASS(BlueprintType)
+class NNEDENOISER_API UNNEDenoiserModelData : public UDataAsset
+{
+	GENERATED_BODY()
+
+public:
+	/** NNE model data */
+	UPROPERTY(EditAnywhere, BlueprintReadWrite, Category=NNEDenoiser)
+	TSoftObjectPtr<UNNEModelData> ModelData;
+
+	/** Input mapping table */
+	UPROPERTY(EditAnywhere, BlueprintReadWrite, Category=NNEDenoiser)
+	TSoftObjectPtr<UDataTable> InputMapping;
+
+	/** Output mapping table */
+	UPROPERTY(EditAnywhere, BlueprintReadWrite, Category=NNEDenoiser)
+	TSoftObjectPtr<UDataTable> OutputMapping;
+};
\ No newline at end of file
diff --git a/Engine/Plugins/Experimental/NNEDenoiser/Source/NNEDenoiser/Public/NNEDenoiserModelIOMappingData.h b/Engine/Plugins/Experimental/NNEDenoiser/Source/NNEDenoiser/Public/NNEDenoiserModelIOMappingData.h
new file mode 100644
--- /dev/null
+++ b/Engine/Plugins/Experimental/NNEDenoiser/Source/NNEDenoiser/Public/NNEDenoiserModelIOMappingData.h
@@ -0,0 +1,46 @@
+// Copyright Epic Games, Inc. All Rights Reserved.
+
+#pragma once
+
+#include "Engine/DataTable.h"
+
+#include "NNEDenoiserModelIOMappingData.generated.h"
+
+/** An enum to represent resource names used for input and output mapping */
+UENUM()
+enum EResourceName : uint8
+{
+	Color,
+	Albedo,
+	Normal,
+	Flow,
+	Output
+};
+
+/** Table row base for denoiser input and output mapping */
+USTRUCT(BlueprintType)
+struct FNNEDenoiserModelIOMappingData : public FTableRowBase
+{
+	GENERATED_USTRUCT_BODY()
+
+public:
+	/** Input/output tensor index */
+	UPROPERTY(EditAnywhere, BlueprintReadWrite, Category=NNEDenoiser)
+	int32 TensorIndex = 0;
+
+	/** Input/output tensor channel */
+	UPROPERTY(EditAnywhere, BlueprintReadWrite, Category=NNEDenoiser)
+	int32 TensorChannel = 0;
+
+	/** Mapped resource name */
+	UPROPERTY(EditAnywhere, BlueprintReadWrite, Category=NNEDenoiser)
+	TEnumAsByte<EResourceName> Resource = EResourceName::Color;
+
+	/** Resource channel */
+	UPROPERTY(EditAnywhere, BlueprintReadWrite, Category=NNEDenoiser)
+	int32 ResourceChannel = 0;
+
+	/** Resource frame index */
+	UPROPERTY(EditAnywhere, BlueprintReadWrite, Category=NNEDenoiser)
+	int32 FrameIndex = 0;
+};
\ No newline at end of file
diff --git a/Engine/Plugins/Experimental/NNEDenoiser/Source/NNEDenoiser/Public/NNEDenoiserSettings.h b/Engine/Plugins/Experimental/NNEDenoiser/Source/NNEDenoiser/Public/NNEDenoiserSettings.h
new file mode 100644
--- /dev/null
+++ b/Engine/Plugins/Experimental/NNEDenoiser/Source/NNEDenoiser/Public/NNEDenoiserSettings.h
@@ -0,0 +1,42 @@
+// Copyright Epic Games, Inc. All Rights Reserved.
+
+#pragma once
+
+#include "Engine/DeveloperSettingsBackedByCVars.h"
+#include "NNEDenoiserModelData.h"
+
+#include "NNEDenoiserSettings.generated.h"
+
+/** An enum to represent denoiser NNE runtime type */
+UENUM()
+enum EDenoiserRuntimeType : uint8
+{
+	CPU,
+	GPU,
+	RDG
+};
+
+/** Settings used to create a NNE Denoise */
+UCLASS(Config = Engine, meta = (DisplayName = "NNE Denoise"))
+class NNEDENOISER_API UNNEDenoiserSettings : public UDeveloperSettingsBackedByCVars
+{
+	GENERATED_BODY()
+
+public:
+	UNNEDenoiserSettings();
+
+	virtual void PostInitProperties() override;
+
+	/** Denoiser model data used to create a NNE Denoise */
+	UPROPERTY(Config, EditAnywhere, Category = "NNE Denoise", meta = (DisplayName = "Denoiser model data", ToolTip = "Select the denoiser model data"))
+	TSoftObjectPtr<UNNEDenoiserModelData> DenoiserModelData;
+
+private:
+	/** Runtime type used to run the NNE Denoise model. Backed by the console variable 'NNEDenoiser.Runtime.Type'. */
+	UPROPERTY(Config, EditAnywhere, Category = "NNE Denoise", meta = (DisplayName = "Runtime Type", ToolTip = "Select a Runtime type", ConsoleVariable = "NNEDenoiser.Runtime.Type"))
+	TEnumAsByte<EDenoiserRuntimeType> RuntimeType;
+
+	/** Runtime name used to run the NNE Denoise model. Backed by the console variable 'NNEDenoiser.Runtime.Name'. */
+	UPROPERTY(Config, EditAnywhere, Category = "NNE Denoise", meta = (DisplayName = "Runtime Name Override", ToolTip = "(Optional) Specify the Runtime name", ConsoleVariable = "NNEDenoiser.Runtime.Name"))
+	FString RuntimeName;
+};
diff --git a/Engine/Plugins/Experimental/NNEDenoiser/Source/NNEDenoiserShaders/Internal/NNEDenoiserShadersDefaultCS.h b/Engine/Plugins/Experimental/NNEDenoiser/Source/NNEDenoiserShaders/Internal/NNEDenoiserShadersDefaultCS.h
new file mode 100644
--- /dev/null
+++ b/Engine/Plugins/Experimental/NNEDenoiser/Source/NNEDenoiserShaders/Internal/NNEDenoiserShadersDefaultCS.h
@@ -0,0 +1,77 @@
+// Copyright Epic Games, Inc. All Rights Reserved.
+
+#pragma once
+
+#include "GlobalShader.h"
+#include "RenderGraphUtils.h"
+#include "ShaderParameterUtils.h"
+
+namespace UE::NNEDenoiserShaders::Internal
+{
+
+	enum class ENNEDenoiserDataType : uint8
+	{
+		None = 0,
+		// Char,
+		// Boolean,
+		Half = 3,
+		Float,
+		MAX
+	};
+
+	class FNNEDenoiserConstants
+	{
+	public:
+		static constexpr int32 THREAD_GROUP_SIZE{ 32 };
+		static constexpr int32 MAX_NUM_MAPPED_CHANNELS{ 4 };
+	};
+	
+	class NNEDENOISERSHADERS_API FNNEDenoiserReadInputCS : public FGlobalShader
+	{
+		DECLARE_GLOBAL_SHADER(FNNEDenoiserReadInputCS);
+		SHADER_USE_PARAMETER_STRUCT(FNNEDenoiserReadInputCS, FGlobalShader)
+
+		class FNNEDenoiserDataType : SHADER_PERMUTATION_ENUM_CLASS("BUFFER_TYPE_INDEX", ENNEDenoiserDataType);
+		class FNNEDenoiserNumMappedChannels : SHADER_PERMUTATION_RANGE_INT("NUM_MAPPED_CHANNELS", 0, FNNEDenoiserConstants::MAX_NUM_MAPPED_CHANNELS);
+		using FPermutationDomain = TShaderPermutationDomain<FNNEDenoiserDataType, FNNEDenoiserNumMappedChannels>;
+
+	public:
+
+		BEGIN_SHADER_PARAMETER_STRUCT(FParameters, )
+			SHADER_PARAMETER(int32, InputTextureWidth)
+			SHADER_PARAMETER(int32, InputTextureHeight)
+			SHADER_PARAMETER_RDG_TEXTURE(Texture2D, InputTexture)
+			SHADER_PARAMETER(int32, InputBufferWidth)
+			SHADER_PARAMETER(int32, InputBufferHeight)
+			SHADER_PARAMETER_RDG_BUFFER_UAV(RWBuffer, InputBuffer)
+			SHADER_PARAMETER_ARRAY(FIntVector4, BufferChannel_TextureChannel_Unused_Unused, [FNNEDenoiserConstants::MAX_NUM_MAPPED_CHANNELS])
+		END_SHADER_PARAMETER_STRUCT()
+
+		static void ModifyCompilationEnvironment(const FGlobalShaderPermutationParameters& InParameters, FShaderCompilerEnvironment& OutEnvironment);
+	};
+
+	class NNEDENOISERSHADERS_API FNNEDenoiserWriteOutputCS : public FGlobalShader
+	{
+		DECLARE_GLOBAL_SHADER(FNNEDenoiserWriteOutputCS);
+		SHADER_USE_PARAMETER_STRUCT(FNNEDenoiserWriteOutputCS, FGlobalShader)
+
+		class FNNEDenoiserDataType : SHADER_PERMUTATION_ENUM_CLASS("BUFFER_TYPE_INDEX", ENNEDenoiserDataType);
+		class FNNEDenoiserNumMappedChannels : SHADER_PERMUTATION_RANGE_INT("NUM_MAPPED_CHANNELS", 0, FNNEDenoiserConstants::MAX_NUM_MAPPED_CHANNELS);
+		using FPermutationDomain = TShaderPermutationDomain<FNNEDenoiserDataType, FNNEDenoiserNumMappedChannels>;
+
+	public:
+
+		BEGIN_SHADER_PARAMETER_STRUCT(FParameters, )
+			SHADER_PARAMETER(int32, OutputBufferWidth)
+			SHADER_PARAMETER(int32, OutputBufferHeight)
+			SHADER_PARAMETER_RDG_BUFFER_UAV(RWBuffer, OutputBuffer)
+			SHADER_PARAMETER(int32, OutputTextureWidth)
+			SHADER_PARAMETER(int32, OutputTextureHeight)
+			SHADER_PARAMETER_RDG_TEXTURE_UAV(RWTexture2D, OutputTexture)
+			SHADER_PARAMETER_ARRAY(FIntVector4, BufferChannel_TextureChannel_Unused_Unused, [FNNEDenoiserConstants::MAX_NUM_MAPPED_CHANNELS])
+		END_SHADER_PARAMETER_STRUCT()
+
+		static void ModifyCompilationEnvironment(const FGlobalShaderPermutationParameters& InParameters, FShaderCompilerEnvironment& OutEnvironment);
+	};
+
+} // namespace UE::NNEDenoiser::Private
\ No newline at end of file
diff --git a/Engine/Plugins/Experimental/NNEDenoiser/Source/NNEDenoiserShaders/Internal/NNEDenoiserShadersOidnCS.h b/Engine/Plugins/Experimental/NNEDenoiser/Source/NNEDenoiserShaders/Internal/NNEDenoiserShadersOidnCS.h
new file mode 100644
--- /dev/null
+++ b/Engine/Plugins/Experimental/NNEDenoiser/Source/NNEDenoiserShaders/Internal/NNEDenoiserShadersOidnCS.h
@@ -0,0 +1,52 @@
+// Copyright Epic Games, Inc. All Rights Reserved.
+
+#pragma once
+
+#include "GlobalShader.h"
+#include "RenderGraphUtils.h"
+#include "ShaderParameterUtils.h"
+
+namespace UE::NNEDenoiserShaders::Internal
+{
+
+	enum class ENNEDenoiserInputKind : uint8
+	{
+		Color = 0,
+		Albedo,
+		Normal,
+		Flow,
+		Output,
+		MAX
+	};
+
+	class FNNEDenoiserOidnConstants
+	{
+	public:
+		static constexpr int32 THREAD_GROUP_SIZE{ 32 };
+	};
+	
+	class NNEDENOISERSHADERS_API FNNEDenoiserOidnCS : public FGlobalShader
+	{
+		DECLARE_GLOBAL_SHADER(FNNEDenoiserOidnCS);
+		SHADER_USE_PARAMETER_STRUCT(FNNEDenoiserOidnCS, FGlobalShader)
+
+		class FNNEDenoiserInputKind : SHADER_PERMUTATION_ENUM_CLASS("INPUT_KIND_INDEX", ENNEDenoiserInputKind);
+		using FPermutationDomain = TShaderPermutationDomain<FNNEDenoiserInputKind>;
+
+	public:
+
+		BEGIN_SHADER_PARAMETER_STRUCT(FParameters, )
+			SHADER_PARAMETER(int32, InputTextureWidth)
+			SHADER_PARAMETER(int32, InputTextureHeight)
+			SHADER_PARAMETER_RDG_TEXTURE(Texture2D, InputTexture)
+			SHADER_PARAMETER(int32, OutputTextureWidth)
+			SHADER_PARAMETER(int32, OutputTextureHeight)
+			SHADER_PARAMETER_RDG_TEXTURE_UAV(RWTexture2D, OutputTexture)
+			SHADER_PARAMETER(float, NormScale)
+			SHADER_PARAMETER(float, InvNormScale)
+		END_SHADER_PARAMETER_STRUCT()
+
+		static void ModifyCompilationEnvironment(const FGlobalShaderPermutationParameters& InParameters, FShaderCompilerEnvironment& OutEnvironment);
+	};
+
+} // namespace UE::NNEDenoiser::Private
\ No newline at end of file
diff --git a/Engine/Plugins/Experimental/NNEDenoiser/Source/NNEDenoiserShaders/NNEDenoiserShaders.Build.cs b/Engine/Plugins/Experimental/NNEDenoiser/Source/NNEDenoiserShaders/NNEDenoiserShaders.Build.cs
new file mode 100644
--- /dev/null
+++ b/Engine/Plugins/Experimental/NNEDenoiser/Source/NNEDenoiserShaders/NNEDenoiserShaders.Build.cs
@@ -0,0 +1,26 @@
+// Copyright Epic Games, Inc. All Rights Reserved.
+
+using UnrealBuildTool;
+
+public class NNEDenoiserShaders : ModuleRules
+{
+	public NNEDenoiserShaders(ReadOnlyTargetRules Target) : base(Target)
+	{
+		PCHUsage = PCHUsageMode.UseExplicitOrSharedPCHs;
+
+		PublicDependencyModuleNames.AddRange(
+			new string[]
+			{
+				"Core",
+				"RenderCore"
+			});
+			
+		PrivateDependencyModuleNames.AddRange(
+			new string[]
+			{
+				"Projects",
+				"Renderer",
+				"RHI"
+			});
+	}
+}
diff --git a/Engine/Plugins/Experimental/NNEDenoiser/Source/NNEDenoiserShaders/Private/NNEDenoiserShadersDefaultCS.cpp b/Engine/Plugins/Experimental/NNEDenoiser/Source/NNEDenoiserShaders/Private/NNEDenoiserShadersDefaultCS.cpp
new file mode 100644
--- /dev/null
+++ b/Engine/Plugins/Experimental/NNEDenoiser/Source/NNEDenoiserShaders/Private/NNEDenoiserShadersDefaultCS.cpp
@@ -0,0 +1,24 @@
+// Copyright Epic Games, Inc. All Rights Reserved.
+
+#include "NNEDenoiserShadersDefaultCS.h"
+
+namespace UE::NNEDenoiserShaders::Internal
+{
+	void FNNEDenoiserReadInputCS::ModifyCompilationEnvironment(const FGlobalShaderPermutationParameters& InParameters, FShaderCompilerEnvironment& OutEnvironment)
+	{
+		FGlobalShader::ModifyCompilationEnvironment(InParameters, OutEnvironment);
+		OutEnvironment.SetDefine(TEXT("THREAD_GROUP_SIZE"), FNNEDenoiserConstants::THREAD_GROUP_SIZE);
+		OutEnvironment.SetDefine(TEXT("MAX_NUM_MAPPED_CHANNELS"), FNNEDenoiserConstants::MAX_NUM_MAPPED_CHANNELS);
+	}
+
+	void FNNEDenoiserWriteOutputCS::ModifyCompilationEnvironment(const FGlobalShaderPermutationParameters& InParameters, FShaderCompilerEnvironment& OutEnvironment)
+	{
+		FGlobalShader::ModifyCompilationEnvironment(InParameters, OutEnvironment);
+		OutEnvironment.SetDefine(TEXT("THREAD_GROUP_SIZE"), FNNEDenoiserConstants::THREAD_GROUP_SIZE);
+		OutEnvironment.SetDefine(TEXT("MAX_NUM_MAPPED_CHANNELS"), FNNEDenoiserConstants::MAX_NUM_MAPPED_CHANNELS);
+	}
+
+	IMPLEMENT_GLOBAL_SHADER(FNNEDenoiserReadInputCS, "/NNEDenoiserShaders/NNEDenoiserShadersDefault.usf", "ReadInput", SF_Compute);
+	IMPLEMENT_GLOBAL_SHADER(FNNEDenoiserWriteOutputCS, "/NNEDenoiserShaders/NNEDenoiserShadersDefault.usf", "WriteOutput", SF_Compute);
+
+} // UE::NNEDenoiser::Private
\ No newline at end of file
diff --git a/Engine/Plugins/Experimental/NNEDenoiser/Source/NNEDenoiserShaders/Private/NNEDenoiserShadersLog.h b/Engine/Plugins/Experimental/NNEDenoiser/Source/NNEDenoiserShaders/Private/NNEDenoiserShadersLog.h
new file mode 100644
--- /dev/null
+++ b/Engine/Plugins/Experimental/NNEDenoiser/Source/NNEDenoiserShaders/Private/NNEDenoiserShadersLog.h
@@ -0,0 +1,8 @@
+// Copyright Epic Games, Inc. All Rights Reserved.
+
+#pragma once
+
+#include "CoreMinimal.h"
+#include "Logging/LogMacros.h"
+
+DECLARE_LOG_CATEGORY_EXTERN(LogNNEDenoiserShaders, Log, All);
\ No newline at end of file
diff --git a/Engine/Plugins/Experimental/NNEDenoiser/Source/NNEDenoiserShaders/Private/NNEDenoiserShadersModule.cpp b/Engine/Plugins/Experimental/NNEDenoiser/Source/NNEDenoiserShaders/Private/NNEDenoiserShadersModule.cpp
new file mode 100644
--- /dev/null
+++ b/Engine/Plugins/Experimental/NNEDenoiser/Source/NNEDenoiserShaders/Private/NNEDenoiserShadersModule.cpp
@@ -0,0 +1,42 @@
+// Copyright Epic Games, Inc. All Rights Reserved.
+
+#include "CoreMinimal.h"
+#include "Interfaces/IPluginManager.h"
+#include "Misc/Paths.h"
+#include "Modules/ModuleManager.h"
+#include "ShaderCore.h"
+#include "NNEDenoiserShadersLog.h"
+
+DEFINE_LOG_CATEGORY(LogNNEDenoiserShaders);
+
+class FNNEDenoiserShadersModule : public IModuleInterface
+{
+public:
+	// Begin IModuleInterface
+	virtual void StartupModule() override
+	{
+		FString BaseDir;
+
+		const TSharedPtr<IPlugin> Plugin = IPluginManager::Get().FindPlugin(TEXT("NNEDenoiser"));
+		
+		if (Plugin.IsValid())
+		{
+			BaseDir = Plugin->GetBaseDir() + TEXT("/Source/NNEDenoiserShaders");
+		}
+		else
+		{
+			UE_LOG(LogNNEDenoiserShaders, Warning, TEXT("Shaders directory not added. Failed to find NNEDenoiser plugin"));
+		}
+
+		FString ModuleShaderDir = FPaths::Combine(BaseDir, TEXT("Shaders"));
+		
+		AddShaderSourceDirectoryMapping(TEXT("/NNEDenoiserShaders"), ModuleShaderDir);
+	}
+
+	virtual void ShutdownModule() override
+	{
+
+	}
+};
+
+IMPLEMENT_MODULE(FNNEDenoiserShadersModule, NNEDenoiserShaders);
diff --git a/Engine/Plugins/Experimental/NNEDenoiser/Source/NNEDenoiserShaders/Private/NNEDenoiserShadersOidn.cpp b/Engine/Plugins/Experimental/NNEDenoiser/Source/NNEDenoiserShaders/Private/NNEDenoiserShadersOidn.cpp
new file mode 100644
--- /dev/null
+++ b/Engine/Plugins/Experimental/NNEDenoiser/Source/NNEDenoiserShaders/Private/NNEDenoiserShadersOidn.cpp
@@ -0,0 +1,15 @@
+// Copyright Epic Games, Inc. All Rights Reserved.
+
+#include "NNEDenoiserShadersOidnCS.h"
+
+namespace UE::NNEDenoiserShaders::Internal
+{
+	void FNNEDenoiserOidnCS::ModifyCompilationEnvironment(const FGlobalShaderPermutationParameters& InParameters, FShaderCompilerEnvironment& OutEnvironment)
+	{
+		FGlobalShader::ModifyCompilationEnvironment(InParameters, OutEnvironment);
+		OutEnvironment.SetDefine(TEXT("THREAD_GROUP_SIZE"), FNNEDenoiserOidnConstants::THREAD_GROUP_SIZE);
+	}
+
+	IMPLEMENT_GLOBAL_SHADER(FNNEDenoiserOidnCS, "/NNEDenoiserShaders/NNEDenoiserShadersOidn.usf", "PreOrPostprocess", SF_Compute);
+
+} // UE::NNEDenoiser::Private
\ No newline at end of file
diff --git a/Engine/Plugins/Experimental/NNEDenoiser/Source/NNEDenoiserShaders/Shaders/NNEDenoiserShadersDefault.usf b/Engine/Plugins/Experimental/NNEDenoiser/Source/NNEDenoiserShaders/Shaders/NNEDenoiserShadersDefault.usf
new file mode 100644
--- /dev/null
+++ b/Engine/Plugins/Experimental/NNEDenoiser/Source/NNEDenoiserShaders/Shaders/NNEDenoiserShadersDefault.usf
@@ -0,0 +1,84 @@
+// Copyright Epic Games, Inc. All Rights Reserved.
+
+#include "/Engine/Public/Platform.ush"
+#include "/Engine/Private/Common.ush"
+
+#if BUFFER_TYPE_INDEX == 4
+	#define BUFFER_TYPE half
+	#define READ(X) float(X)
+	#define WRITE(X) half(X)
+#else
+	// BUFFER_TYPE_INDEX == 3
+	#define BUFFER_TYPE float
+	#define READ(X) X
+	#define WRITE(X) X
+#endif
+
+int InputTextureWidth;
+int InputTextureHeight;
+Texture2D InputTexture;
+
+int InputBufferWidth;
+int InputBufferHeight;
+RWBuffer<BUFFER_TYPE> InputBuffer;
+
+int4 BufferChannel_TextureChannel_Unused_Unused[MAX_NUM_MAPPED_CHANNELS];
+
+[numthreads(THREAD_GROUP_SIZE, THREAD_GROUP_SIZE, 1)]
+void ReadInput(in const uint3 DispatchThreadID : SV_DispatchThreadID)
+{
+	if (DispatchThreadID.x >= InputBufferWidth || DispatchThreadID.y >= InputBufferHeight)
+	{
+		return;
+	}
+
+	const int Idx = DispatchThreadID.y * InputBufferWidth + DispatchThreadID.x;
+	const int Offset = InputBufferWidth * InputBufferHeight;
+
+	const float4 BufferColor = InputTexture.Load(int3(DispatchThreadID.x, DispatchThreadID.y, 0));
+
+	for (int I = 0; I < NUM_MAPPED_CHANNELS; I++)
+	{
+		const int4 Info = BufferChannel_TextureChannel_Unused_Unused[I];
+		const int BufferChannel = Info.x;
+		const int TextureChannel = Info.y;
+
+		InputBuffer[Idx + BufferChannel * Offset] = WRITE(BufferColor[TextureChannel]);
+	}
+}
+
+int OutputBufferWidth;
+int OutputBufferHeight;
+RWBuffer<BUFFER_TYPE> OutputBuffer;
+
+int OutputTextureWidth;
+int OutputTextureHeight;
+RWTexture2D<float4> OutputTexture;
+
+[numthreads(THREAD_GROUP_SIZE, THREAD_GROUP_SIZE, 1)]
+void WriteOutput(in const uint3 DispatchThreadID : SV_DispatchThreadID)
+{
+	if (DispatchThreadID.x >= OutputTextureWidth || DispatchThreadID.y >= OutputTextureHeight)
+	{
+		return;
+	}
+
+	// nearest neighbour
+	const int X = (int)(DispatchThreadID.x * (OutputBufferWidth - 1.0) / (OutputTextureWidth - 1.0));
+	const int Y = (int)(DispatchThreadID.y * (OutputBufferHeight - 1.0) / (OutputTextureHeight - 1.0));
+
+	const int Idx = Y * OutputBufferWidth + X;
+	const int Offset = OutputBufferWidth * OutputBufferHeight;
+
+	for (int I = 0; I < NUM_MAPPED_CHANNELS; I++)
+	{
+		const int4 Info = BufferChannel_TextureChannel_Unused_Unused[I];
+		const int BufferChannel = Info.x;
+		const int TextureChannel = Info.y;
+
+		OutputTexture[DispatchThreadID.xy][TextureChannel] = READ(OutputBuffer[Idx + BufferChannel * Offset]);
+	}
+
+	// note: alpha channel denoising or transfer not supported (yet)
+	OutputTexture[DispatchThreadID.xy].a = 1.0;
+}
\ No newline at end of file
diff --git a/Engine/Plugins/Experimental/NNEDenoiser/Source/NNEDenoiserShaders/Shaders/NNEDenoiserShadersOidn.usf b/Engine/Plugins/Experimental/NNEDenoiser/Source/NNEDenoiserShaders/Shaders/NNEDenoiserShadersOidn.usf
new file mode 100644
--- /dev/null
+++ b/Engine/Plugins/Experimental/NNEDenoiser/Source/NNEDenoiserShaders/Shaders/NNEDenoiserShadersOidn.usf
@@ -0,0 +1,85 @@
+// Copyright Epic Games, Inc. All Rights Reserved.
+
+#include "/Engine/Public/Platform.ush"
+#include "/Engine/Private/Common.ush"
+
+int InputTextureWidth;
+int InputTextureHeight;
+Texture2D InputTexture;
+
+int OutputTextureWidth;
+int OutputTextureHeight;
+RWTexture2D<float4> OutputTexture;
+
+float NormScale;
+float InvNormScale;
+
+#define MAX_FLT 10000.0
+
+///
+
+// from [REDACTED_PROFILE_URL]/oidn/blob/master/core/color.h
+// struct FPU
+// {
+	static float a  =  1.41283765e+03;
+	static float b  =  1.64593172e+00;
+	static float c  =  4.31384981e-01;
+	static float d  = -2.94139609e-03;
+	static float e  =  1.92653254e-01;
+	static float f  =  6.26026094e-03;
+	static float g  =  9.98620152e-01;
+	static float y0 =  1.57945760e-06;
+	static float y1 =  3.22087631e-02;
+	static float x0 =  2.23151711e-03;
+	static float x1 =  3.70974749e-01;
+
+	static float Forward(float y)
+	{
+		if (y <= y0) return a * y;
+		else if (y <= y1) return b * pow(y, c) + d;
+		return e * log(y + f) + g;
+	}
+
+	static float Inverse(float x)
+	{
+		if (x <= x0) return x / a;
+		else if (x <= x1) return pow((x - d) / b, 1.0 / c);
+		return exp((x - g) / e) - f;
+	}
+
+	float3 GetProcessedValue(float3 Value)
+	{
+#if INPUT_KIND_INDEX == 0
+		return float3(Forward(Value.r), Forward(Value.g), Forward(Value.b)) * NormScale;
+#elif INPUT_KIND_INDEX == 1
+		return clamp(Value, float3(0.0, 0.0, 0.0), float3(1.0, 1.0, 1.0));
+#elif INPUT_KIND_INDEX == 2
+		return clamp(Value, float3(-1.0, -1.0, -1.0), float3(1.0, 1.0, 1.0)) * 0.5 + float3(0.5, 0.5, 0.5);
+#elif INPUT_KIND_INDEX == 3
+		// TODO flow texture check for any necessary conversion
+		return Value;
+#else
+		// INPUT_KIND_INDEX == 4 // Output
+		Value = clamp(Value, float3(0.0, 0.0, 0.0), float3(MAX_FLT, MAX_FLT, MAX_FLT)) * InvNormScale;
+		
+		return float3(Inverse(Value.r), Inverse(Value.g), Inverse(Value.b));
+#endif
+	}
+
+// };
+
+///
+
+[numthreads(THREAD_GROUP_SIZE, THREAD_GROUP_SIZE, 1)]
+void PreOrPostprocess(in const uint3 DispatchThreadID : SV_DispatchThreadID)
+{
+	if (DispatchThreadID.x >= OutputTextureWidth || DispatchThreadID.y >= OutputTextureHeight)
+	{
+		return;
+	}
+
+	const float4 Color = InputTexture.Load(int3(DispatchThreadID.x, DispatchThreadID.y, 0));
+
+	OutputTexture[DispatchThreadID.xy].rgb = GetProcessedValue(Color.rgb);
+	OutputTexture[DispatchThreadID.xy].a = Color.a;
+}
\ No newline at end of file
