[NNEDenoiser] Updated default settings.

Tests:
 * Non-Unity build
 * Verified new default settings in UnrealEditor

#rnx

diff --git a/Engine/Plugins/Experimental/NNEDenoiser/Source/NNEDenoiser/Private/NNEDenoiserModelInstanceRDG.cpp b/Engine/Plugins/Experimental/NNEDenoiser/Source/NNEDenoiser/Private/NNEDenoiserModelInstanceRDG.cpp
--- a/Engine/Plugins/Experimental/NNEDenoiser/Source/NNEDenoiser/Private/NNEDenoiserModelInstanceRDG.cpp
+++ b/Engine/Plugins/Experimental/NNEDenoiser/Source/NNEDenoiser/Private/NNEDenoiserModelInstanceRDG.cpp
@@ -16,7 +16,7 @@
 namespace UE::NNEDenoiser::Private
 {
 
-	static const FString DefaultRuntimeRDGName = TEXT("NNERuntimeRDGDml");
+	static const FString DefaultRuntimeRDGName = TEXT("NNERuntimeRDGHlsl");
 
 	TUniquePtr<FModelInstanceRDG> FModelInstanceRDG::Make(UNNEModelData& ModelData, const FString& RuntimeNameOverride)
 	{
diff --git a/Engine/Plugins/Experimental/NNEDenoiser/Source/NNEDenoiser/Private/NNEDenoiserSettings.cpp b/Engine/Plugins/Experimental/NNEDenoiser/Source/NNEDenoiser/Private/NNEDenoiserSettings.cpp
--- a/Engine/Plugins/Experimental/NNEDenoiser/Source/NNEDenoiser/Private/NNEDenoiserSettings.cpp
+++ b/Engine/Plugins/Experimental/NNEDenoiser/Source/NNEDenoiser/Private/NNEDenoiserSettings.cpp
@@ -11,7 +11,7 @@ UNNEDenoiserSettings::UNNEDenoiserSettings()
 	CategoryName = TEXT("Plugins");
 	SectionName = TEXT("NNE Denoiser");
 
-	DenoiserModelData = FSoftObjectPath(TEXT("/NNEDenoiser/NNEDNN_Oidn2_Color_720.NNEDNN_Oidn2_Color_720"));
+	DenoiserModelData = FSoftObjectPath(TEXT("/NNEDenoiser/NNEDNN_Oidn2_ColorAlbedoNormal_720.NNEDNN_Oidn2_ColorAlbedoNormal_720"));
 }
 
 void UNNEDenoiserSettings::PostInitProperties()
diff --git a/Engine/Plugins/Experimental/NNEDenoiser/Source/NNEDenoiser/Private/NNEDenoiserViewExtension.cpp b/Engine/Plugins/Experimental/NNEDenoiser/Source/NNEDenoiser/Private/NNEDenoiserViewExtension.cpp
--- a/Engine/Plugins/Experimental/NNEDenoiser/Source/NNEDenoiser/Private/NNEDenoiserViewExtension.cpp
+++ b/Engine/Plugins/Experimental/NNEDenoiser/Source/NNEDenoiser/Private/NNEDenoiserViewExtension.cpp
@@ -40,7 +40,7 @@ static TAutoConsoleVariable<int32> CVarNNEDenoiserModelData(
 
 static TAutoConsoleVariable<int32> CVarNNEDenoiserRuntimeType(
 	TEXT("NNEDenoiser.Runtime.Type"),
-	0,
+	1,
 	TEXT("Defines the runtime type to run the denoiser model.\n")
 	TEXT("  0: CPU\n")
 	TEXT("  1: GPU\n")
