[NNEDenoiser] Lazy load neural model (i.e. properly support headless agents).

Problem:
Should not throw any errors on module load.

Solution:
Lazy load neural model initialization of denoiser on begin render.

Tests:
* Non-Unity build
* UnrealEditor, Win64: Startup with NNEDenoiser plugin enabled

#rnx

diff --git a/Engine/Plugins/Experimental/NNEDenoiser/Source/NNEDenoiser/Private/NNEDenoiserViewExtension.cpp b/Engine/Plugins/Experimental/NNEDenoiser/Source/NNEDenoiser/Private/NNEDenoiserViewExtension.cpp
--- a/Engine/Plugins/Experimental/NNEDenoiser/Source/NNEDenoiser/Private/NNEDenoiserViewExtension.cpp
+++ b/Engine/Plugins/Experimental/NNEDenoiser/Source/NNEDenoiser/Private/NNEDenoiserViewExtension.cpp
@@ -40,7 +40,7 @@ static TAutoConsoleVariable<int32> CVarNNEDenoiserModelData(
 
 static TAutoConsoleVariable<int32> CVarNNEDenoiserRuntimeType(
 	TEXT("NNEDenoiser.Runtime.Type"),
-	0,
+	1,
 	TEXT("Defines the runtime type to run the denoiser model.\n")
 	TEXT("  0: CPU\n")
 	TEXT("  1: GPU\n")
@@ -216,8 +216,6 @@ TUniquePtr<FGenericDenoiser> CreateNNEDenoiserFromAsset(const FString& AssetName
 
 FViewExtension::FViewExtension(const FAutoRegister& AutoRegister) : FSceneViewExtensionBase(AutoRegister)
 {
-	ApplySettings(GetDefault<UNNEDenoiserSettings>());
-
 #if WITH_EDITOR
 	GetMutableDefault<UNNEDenoiserSettings>()->OnSettingChanged().AddRaw(this, &FViewExtension::OnDenoiserSettingsChanged);
 #endif
