[NNEDenoiser] Set NNERuntimeORTDml on RDG as default.

#test Start NNEDenoiserTest in Editor and verify ORTDml on RDG being set as default.

#rnx

diff --git a/Engine/Plugins/Experimental/NNEDenoiser/Source/NNEDenoiser/Private/NNEDenoiserModelInstanceRDG.cpp b/Engine/Plugins/Experimental/NNEDenoiser/Source/NNEDenoiser/Private/NNEDenoiserModelInstanceRDG.cpp
--- a/Engine/Plugins/Experimental/NNEDenoiser/Source/NNEDenoiser/Private/NNEDenoiserModelInstanceRDG.cpp
+++ b/Engine/Plugins/Experimental/NNEDenoiser/Source/NNEDenoiser/Private/NNEDenoiserModelInstanceRDG.cpp
@@ -16,9 +16,12 @@
 namespace UE::NNEDenoiser::Private
 {
 
-	static const FString DefaultRuntimeRDGName = TEXT("NNERuntimeRDGHlsl");
+	static const FString DefaultRuntimeRDGName = TEXT("NNERuntimeORTDml");
+	static const FString FallbackRuntimeRDGName = TEXT("NNERuntimeRDGHlsl");
 
-	TUniquePtr<FModelInstanceRDG> FModelInstanceRDG::Make(UNNEModelData& ModelData, const FString& RuntimeNameOverride)
+namespace Internal
+{
+	TUniquePtr<FModelInstanceRDG> Make(UNNEModelData& ModelData, const FString& RuntimeNameOverride)
 	{
 		FString RuntimeRDGName = RuntimeNameOverride.IsEmpty() ? DefaultRuntimeRDGName : RuntimeNameOverride;
 		
@@ -52,6 +55,22 @@ namespace UE::NNEDenoiser::Private
 
 		return MakeUnique<FModelInstanceRDG>(ModelInstance.ToSharedRef());
 	}
+} // namespace Internal
+
+	TUniquePtr<FModelInstanceRDG> FModelInstanceRDG::Make(UNNEModelData& ModelData, const FString& RuntimeNameOverride)
+	{
+		if (!RuntimeNameOverride.IsEmpty())
+		{
+			return Internal::Make(ModelData, RuntimeNameOverride);
+		}
+
+		if (TUniquePtr<FModelInstanceRDG> Result = Internal::Make(ModelData, DefaultRuntimeRDGName); Result.IsValid())
+		{
+			return Result;
+		}
+
+		return Internal::Make(ModelData, FallbackRuntimeRDGName);
+	}
 
 	FModelInstanceRDG::FModelInstanceRDG(TSharedRef<UE::NNE::IModelInstanceRDG> ModelInstance) :
 		ModelInstance(ModelInstance)
diff --git a/Engine/Plugins/Experimental/NNEDenoiser/Source/NNEDenoiser/Private/NNEDenoiserViewExtension.cpp b/Engine/Plugins/Experimental/NNEDenoiser/Source/NNEDenoiser/Private/NNEDenoiserViewExtension.cpp
--- a/Engine/Plugins/Experimental/NNEDenoiser/Source/NNEDenoiser/Private/NNEDenoiserViewExtension.cpp
+++ b/Engine/Plugins/Experimental/NNEDenoiser/Source/NNEDenoiser/Private/NNEDenoiserViewExtension.cpp
@@ -40,7 +40,7 @@ static TAutoConsoleVariable<int32> CVarNNEDenoiserModelData(
 
 static TAutoConsoleVariable<int32> CVarNNEDenoiserRuntimeType(
 	TEXT("NNEDenoiser.Runtime.Type"),
-	1,
+	2,
 	TEXT("Defines the runtime type to run the denoiser model.\n")
 	TEXT("  0: CPU\n")
 	TEXT("  1: GPU\n")
