[NNEDenoiser] Added support for dynamic sized networks and improved tiling. Also replaced fixed size models with all dynamic size models.

#tests Editor start and run denoiser models (all dynamic size models now)

#virtualized

diff --git a/Engine/Plugins/Experimental/NNEDenoiser/Source/NNEDenoiser/Private/NNEDenoiserGenericDenoiser.cpp b/Engine/Plugins/Experimental/NNEDenoiser/Source/NNEDenoiser/Private/NNEDenoiserGenericDenoiser.cpp
--- a/Engine/Plugins/Experimental/NNEDenoiser/Source/NNEDenoiser/Private/NNEDenoiserGenericDenoiser.cpp
+++ b/Engine/Plugins/Experimental/NNEDenoiser/Source/NNEDenoiser/Private/NNEDenoiserGenericDenoiser.cpp
@@ -77,11 +77,40 @@ void RegisterTextureIfNeeded(FResourceManager& ResourceManager, FRDGTextureRef T
 	ResourceManager.AddTexture(ResourceName, Texture, NumFrames);
 }
 
-bool PrepareAndValidate(IModelInstance& ModelInstance, const IInputProcess& InputProcess, const IOutputProcess& OutputProcess, FIntPoint Extent)
+TArray<NNE::FTensorShape> GetOutputTensorShapes(IModelInstance& ModelInstance)
 {
-	bool Result = InputProcess.PrepareAndValidate(ModelInstance, Extent);
-	Result &= OutputProcess.Validate(ModelInstance, Extent);
+	TConstArrayView<NNE::FTensorShape> InputShapes = ModelInstance.GetInputTensorShapes();
+	check(!InputShapes.IsEmpty());
 
+	TConstArrayView<uint32> InputShapeData = InputShapes[0].GetData();
+	const uint32 Width = InputShapeData[3];
+	const uint32 Height = InputShapeData[2];
+
+	TConstArrayView<NNE::FTensorDesc> OutputDescs = ModelInstance.GetOutputTensorDescs();
+	TConstArrayView<NNE::FTensorShape> OutputShapes = ModelInstance.GetOutputTensorShapes();
+
+	if (OutputDescs.Num() == OutputShapes.Num())
+	{
+		return TArray<NNE::FTensorShape>(OutputShapes);
+	}
+
+	// If output shapes not set yet, try to manually resolve, otherwise they need to be user specified.
+	TArray<NNE::FTensorShape> Result;
+	for (int32 i = 0; i < OutputDescs.Num(); i++)
+	{
+		check(OutputDescs[i].GetShape().Rank() == 4);
+
+		TConstArrayView<int32> SymbolicShapeData = OutputDescs[i].GetShape().GetData();
+		
+		TArray<uint32> ShapeData;
+		ShapeData.SetNumUninitialized(OutputDescs[i].GetShape().Rank());
+		ShapeData[0] = (uint32)SymbolicShapeData[0];
+		ShapeData[1] = (uint32)SymbolicShapeData[1];
+		ShapeData[2] = (uint32)Height;
+		ShapeData[3] = (uint32)Width;
+
+		Result.Add(NNE::FTensorShape::Make(ShapeData));
+	}
 	return Result;
 }
 
@@ -99,7 +128,7 @@ void AddTilePasses(FRDGBuilder& GraphBuilder, IModelInstance& ModelInstance, con
 
 	// 3. Write output based on output buffer
 	FRDGTextureRef OutputTexture = ResourceManager.GetTexture(EResourceName::Output, 0);
-	OutputProcess.AddPasses(GraphBuilder, ModelInstance.GetOutputTensorDescs(), ModelInstance.GetOutputTensorShapes(), ResourceAccess, OutputBuffers, OutputTexture);
+	OutputProcess.AddPasses(GraphBuilder, ModelInstance.GetOutputTensorDescs(), GetOutputTensorShapes(ModelInstance), ResourceAccess, OutputBuffers, OutputTexture);
 }
 
 FGenericDenoiser::FGenericDenoiser(TUniquePtr<IModelInstance> ModelInstance, TUniquePtr<IInputProcess> InputProcess, TUniquePtr<IOutputProcess> OutputProcess, FParameters DenoiserParameters) :
@@ -113,6 +142,40 @@ FGenericDenoiser::~FGenericDenoiser()
 
 }
 
+bool FGenericDenoiser::Prepare(FIntPoint Extent)
+{
+	if (Extent == LastExtent)
+	{
+		return true;
+	}
+
+	// Probably would be enough to do this only once at the very beginning...
+	// We just want to be sure that everything up to width and height is correct
+	if (!InputProcess->Validate(*ModelInstance, {-1, -1}))
+	{
+		return false;
+	}
+
+	TConstArrayView<int32> SymbolicInputShape = ModelInstance->GetInputTensorDescs()[0].GetShape().GetData();
+	const FIntPoint TargetTileSize = {SymbolicInputShape[3], SymbolicInputShape[2]};
+
+	Tiling = CreateTiling(TargetTileSize,
+				DenoiserParameters.TilingConfig.MaxSize,
+				DenoiserParameters.TilingConfig.MinSize,
+				DenoiserParameters.TilingConfig.Alignment,
+				DenoiserParameters.TilingConfig.Overlap,
+				Extent);
+
+	if (!InputProcess->Prepare(*ModelInstance, Tiling.TileSize))
+	{
+		return false;
+	}
+
+	LastExtent = Extent;
+
+	return true;
+}
+
 TUniquePtr<FHistory> FGenericDenoiser::AddPasses(
 	FRDGBuilder& GraphBuilder,
 	FRDGTextureRef ColorTex,
@@ -137,56 +200,48 @@ TUniquePtr<FHistory> FGenericDenoiser::AddPasses(
 
 	TUniquePtr<FHistory> Result;
 
-	if (Extent == LastExtent || PrepareAndValidate(*ModelInstance, *InputProcess, *OutputProcess, Extent))
+	if (!Prepare(Extent))
 	{
-		LastExtent = Extent;
-
-		const TConstArrayView<uint32> InputShapeData = ModelInstance->GetInputTensorShapes()[0].GetData();
-		const FIntPoint ModelInputSize{(int32)InputShapeData[3], (int32)InputShapeData[2]};
-
-		const FTiling Tiling = CreateTiling(ModelInputSize, DenoiserParameters.TileMinimumOverlap, Extent);
+		AddCopyTexturePass(GraphBuilder, ColorTex, OutputTex, FRHICopyTextureInfo{});
 
-		TMap<EResourceName, TArray<TRefCountPtr<IPooledRenderTarget>>> ResourceMap;
-		if (History)
-		{
-			ResourceMap = History->GetResourceMap();
-		}
+		return Result;
+	}
 
-		FResourceManager ResourceManager(GraphBuilder, Tiling, ResourceMap);
-		RegisterTextureIfNeeded(ResourceManager, ColorTex, EResourceName::Color, *InputProcess);
-		RegisterTextureIfNeeded(ResourceManager, AlbedoTex, EResourceName::Albedo, *InputProcess);
-		RegisterTextureIfNeeded(ResourceManager, NormalTex, EResourceName::Normal, *InputProcess);
-		if (FlowTex != FRDGTextureRef{})
-		{
-			RegisterTextureIfNeeded(ResourceManager, FlowTex, EResourceName::Flow, *InputProcess);
-		}
-		RegisterTextureIfNeeded(ResourceManager, OutputTex, EResourceName::Output, *InputProcess);
+	TMap<EResourceName, TArray<TRefCountPtr<IPooledRenderTarget>>> ResourceMap;
+	if (History)
+	{
+		ResourceMap = History->GetResourceMap();
+	}
 
-		TArray<FRDGBufferRef> InputBuffers = CreateBuffersRDG(GraphBuilder, ModelInstance->GetInputTensorDescs(), ModelInstance->GetInputTensorShapes());
-		TArray<FRDGBufferRef> OutputBuffers = CreateBuffersRDG(GraphBuilder, ModelInstance->GetOutputTensorDescs(), ModelInstance->GetOutputTensorShapes());
+	FResourceManager ResourceManager(GraphBuilder, Tiling, ResourceMap);
+	RegisterTextureIfNeeded(ResourceManager, ColorTex, EResourceName::Color, *InputProcess);
+	RegisterTextureIfNeeded(ResourceManager, AlbedoTex, EResourceName::Albedo, *InputProcess);
+	RegisterTextureIfNeeded(ResourceManager, NormalTex, EResourceName::Normal, *InputProcess);
+	if (FlowTex != FRDGTextureRef{})
+	{
+		RegisterTextureIfNeeded(ResourceManager, FlowTex, EResourceName::Flow, *InputProcess);
+	}
+	RegisterTextureIfNeeded(ResourceManager, OutputTex, EResourceName::Output, *InputProcess);
 
-		UE_LOG(LogNNEDenoiser, Log, TEXT("Divided work into %dx%d tiles of size %dx%d each..."), Tiling.Count.X, Tiling.Count.Y, Tiling.TileSize.X, Tiling.TileSize.Y);
+	TArray<FRDGBufferRef> InputBuffers = CreateBuffersRDG(GraphBuilder, ModelInstance->GetInputTensorDescs(), ModelInstance->GetInputTensorShapes());
+	TArray<FRDGBufferRef> OutputBuffers = CreateBuffersRDG(GraphBuilder, ModelInstance->GetOutputTensorDescs(), GetOutputTensorShapes(*ModelInstance));
 
-		for (int32 I = 0; I < Tiling.Tiles.Num(); I++)
-		{
-			ResourceManager.BeginTile(I);
+	UE_LOG(LogNNEDenoiser, Log, TEXT("Divided work of size %dx%d into %dx%d tiles of size %dx%d each..."), Extent.X, Extent.Y, Tiling.Count.X, Tiling.Count.Y, Tiling.TileSize.X, Tiling.TileSize.Y);
 
-			// Do inference on tile
-			AddTilePasses(GraphBuilder, *ModelInstance, *InputProcess, *OutputProcess, ResourceManager, InputBuffers, OutputBuffers);
+	for (int32 I = 0; I < Tiling.Tiles.Num(); I++)
+	{
+		ResourceManager.BeginTile(I);
 
-			ResourceManager.EndTile();
-		}
+		// Do inference on tile
+		AddTilePasses(GraphBuilder, *ModelInstance, *InputProcess, *OutputProcess, ResourceManager, InputBuffers, OutputBuffers);
 
-		ResourceMap = ResourceManager.MakeHistoryResourceMap();
-		if (!ResourceMap.IsEmpty())
-		{
-			Result = MakeUnique<FHistory>(*DebugName, MoveTemp(ResourceMap));
-		}
+		ResourceManager.EndTile();
 	}
-	else
+
+	ResourceMap = ResourceManager.MakeHistoryResourceMap();
+	if (!ResourceMap.IsEmpty())
 	{
-		// If something went wrong, just copy color to output...
-		AddCopyTexturePass(GraphBuilder, ColorTex, OutputTex, FRHICopyTextureInfo{});
+		Result = MakeUnique<FHistory>(*DebugName, MoveTemp(ResourceMap));
 	}
 
 	return Result;
diff --git a/Engine/Plugins/Experimental/NNEDenoiser/Source/NNEDenoiser/Private/NNEDenoiserGenericDenoiser.h b/Engine/Plugins/Experimental/NNEDenoiser/Source/NNEDenoiser/Private/NNEDenoiserGenericDenoiser.h
--- a/Engine/Plugins/Experimental/NNEDenoiser/Source/NNEDenoiser/Private/NNEDenoiserGenericDenoiser.h
+++ b/Engine/Plugins/Experimental/NNEDenoiser/Source/NNEDenoiser/Private/NNEDenoiserGenericDenoiser.h
@@ -5,6 +5,7 @@
 #include "Containers/UnrealString.h"
 #include "Math/IntPoint.h"
 #include "NNEDenoiserParameters.h"
+#include "NNEDenoiserTiling.h"
 #include "RenderGraphFwd.h"
 #include "Templates/UniquePtr.h"
 
@@ -39,6 +40,8 @@ public:
 		FHistory* History);
 
 private:
+	bool Prepare(FIntPoint Extent);
+
 	inline static const FString DebugName = TEXT("FGenericDenoiser");
 
 	TUniquePtr<IModelInstance> ModelInstance;
@@ -47,6 +50,7 @@ private:
 	FParameters DenoiserParameters;
 
 	FIntPoint LastExtent = { -1, -1 };
+	FTiling Tiling{};
 };
 
 } // namespace UE::NNEDenoiser::Private
\ No newline at end of file
diff --git a/Engine/Plugins/Experimental/NNEDenoiser/Source/NNEDenoiser/Private/NNEDenoiserIOProcess.h b/Engine/Plugins/Experimental/NNEDenoiser/Source/NNEDenoiser/Private/NNEDenoiserIOProcess.h
--- a/Engine/Plugins/Experimental/NNEDenoiser/Source/NNEDenoiser/Private/NNEDenoiserIOProcess.h
+++ b/Engine/Plugins/Experimental/NNEDenoiser/Source/NNEDenoiser/Private/NNEDenoiserIOProcess.h
@@ -24,7 +24,9 @@ class IInputProcess
 public:
 	virtual ~IInputProcess() = default;
 
-	virtual bool PrepareAndValidate(IModelInstance& ModelInstance, FIntPoint Extent) const = 0;
+	virtual bool Validate(const IModelInstance& ModelInstance, FIntPoint Extent) const = 0;
+
+	virtual bool Prepare(IModelInstance& ModelInstance, FIntPoint Extent) const = 0;
 
 	virtual int32 NumFrames(EResourceName Name) const = 0;
 
diff --git a/Engine/Plugins/Experimental/NNEDenoiser/Source/NNEDenoiser/Private/NNEDenoiserIOProcessBase.cpp b/Engine/Plugins/Experimental/NNEDenoiser/Source/NNEDenoiser/Private/NNEDenoiserIOProcessBase.cpp
--- a/Engine/Plugins/Experimental/NNEDenoiser/Source/NNEDenoiser/Private/NNEDenoiserIOProcessBase.cpp
+++ b/Engine/Plugins/Experimental/NNEDenoiser/Source/NNEDenoiser/Private/NNEDenoiserIOProcessBase.cpp
@@ -145,44 +145,84 @@ void AddReadInputPassForKind(
 
 } // namespace IOProcessBaseHelper
 
-bool FInputProcessBase::PrepareAndValidate(IModelInstance& ModelInstance, FIntPoint Extent) const
+bool FInputProcessBase::Validate(const IModelInstance& ModelInstance, FIntPoint Extent) const
 {
-	SCOPED_NAMED_EVENT_TEXT("NNEDenoiser.Prepare", FColor::Magenta);
+	SCOPED_NAMED_EVENT_TEXT("NNEDenoiser.Validate", FColor::Magenta);
+
+	checkf(Extent == FIntPoint(-1, -1) || (Extent.X >= 0 && Extent.Y >= 0), TEXT("Extent should be either fully symbolic or set!"));
 
 	const int32 NumBatches = 1;
 
-	UE_LOG(LogNNEDenoiser, Log, TEXT("Configure model for extent %dx%d..."), Extent.X, Extent.Y);
+	UE_LOG(LogNNEDenoiser, Log, TEXT("Validate model for extent %dx%d..."), Extent.X, Extent.Y);
 
-	TConstArrayView<UE::NNE::FTensorDesc> InputTensorDescs = ModelInstance.GetInputTensorDescs();
-	if (InputTensorDescs.Num() != InputLayout.Num())
+	if (Extent == FIntPoint(-1, -1))
 	{
-		UE_LOG(LogNNEDenoiser, Error, TEXT("Wrong number of inputs (expected %d, got %d)!"), InputLayout.Num(), InputTensorDescs.Num())
-		return false;
-	}
+		TConstArrayView<UE::NNE::FTensorDesc> InputTensorDescs = ModelInstance.GetInputTensorDescs();
+		if (InputTensorDescs.Num() != InputLayout.Num())
+		{
+			UE_LOG(LogNNEDenoiser, Error, TEXT("Wrong number of inputs (expected %d, got %d)!"), InputLayout.Num(), InputTensorDescs.Num())
+			return false;
+		}
 
-	UE_LOG(LogNNEDenoiser, Log, TEXT("Input shapes (set):"));
+		TArray<UE::NNE::FTensorShape> InputShapes;
+		for (int32 Idx = 0; Idx < InputTensorDescs.Num(); Idx++)
+		{
+			TConstArrayView<int32> InputSymbolicTensorShapeData = InputTensorDescs[Idx].GetShape().GetData();
+			const TArray<int32, TInlineAllocator<4>> RequiredInputShapeData = { NumBatches, InputLayout.NumChannels(Idx), Extent.Y, Extent.X };
 
-	TArray<UE::NNE::FTensorShape> InputShapes;
-	for (int32 Idx = 0; Idx < InputTensorDescs.Num(); Idx++)
+			if (!IsTensorShapeValid(InputSymbolicTensorShapeData, RequiredInputShapeData, TEXT("Input")))
+			{
+				return false;
+			}
+		}
+	}
+	else
 	{
-		TConstArrayView<int32> InputSymbolicTensorShapeData = InputTensorDescs[Idx].GetShape().GetData();
-		const TArray<int32, TInlineAllocator<4>> RequiredInputShapeData = { NumBatches, InputLayout.NumChannels(Idx), -1, -1 };
-
-		if (!IsTensorShapeValid(InputSymbolicTensorShapeData, RequiredInputShapeData, TEXT("Input")))
+		TConstArrayView<UE::NNE::FTensorShape> InputTensorShapes = ModelInstance.GetInputTensorShapes();
+		if (InputTensorShapes.Num() != InputLayout.Num())
 		{
+			UE_LOG(LogNNEDenoiser, Error, TEXT("Wrong number of inputs (expected %d, got %d)!"), InputLayout.Num(), InputTensorShapes.Num())
 			return false;
 		}
 
-		const int32 ModelInputWidth = InputSymbolicTensorShapeData[3] >= 0 ? InputSymbolicTensorShapeData[3] : Extent.X;
-		const int32 ModelInputHeight = InputSymbolicTensorShapeData[2] >= 0 ? InputSymbolicTensorShapeData[2] : Extent.Y;
-
-		if (Extent.X < ModelInputWidth || Extent.Y < ModelInputHeight)
+		for (int32 Idx = 0; Idx < InputTensorShapes.Num(); Idx++)
 		{
-			UE_LOG(LogNNEDenoiser, Error, TEXT("Input image too small (model expects at least size %dx%d, got %dx%d!"), ModelInputWidth, ModelInputHeight, Extent.X, Extent.Y);
-			return false;
+			TConstArrayView<uint32> InputTensorShapeData = InputTensorShapes[Idx].GetData();
+			const TArray<int32, TInlineAllocator<4>> RequiredInputShapeData = { NumBatches, InputLayout.NumChannels(Idx), Extent.Y, Extent.X };
+
+			if (!IsTensorShapeValid(InputTensorShapeData, RequiredInputShapeData, TEXT("Input")))
+			{
+				return false;
+			}
 		}
+	}
+
+	return true;
+}
+
+bool FInputProcessBase::Prepare(IModelInstance& ModelInstance, FIntPoint Extent) const
+{
+	SCOPED_NAMED_EVENT_TEXT("NNEDenoiser.Prepare", FColor::Magenta);
+
+	if (!Validate(ModelInstance, {-1, -1}))
+	{
+		return false;
+	}
 
-		const FIntVector4 ModelInputShape = { 1, InputSymbolicTensorShapeData[1], ModelInputHeight, ModelInputWidth };
+	const int32 NumBatches = 1;
+
+	UE_LOG(LogNNEDenoiser, Log, TEXT("Configure model for extent %dx%d..."), Extent.X, Extent.Y);
+
+	TConstArrayView<UE::NNE::FTensorDesc> InputTensorDescs = ModelInstance.GetInputTensorDescs();
+	
+	UE_LOG(LogNNEDenoiser, Log, TEXT("Input shapes (set):"));
+
+	TArray<UE::NNE::FTensorShape> InputShapes;
+	for (int32 Idx = 0; Idx < InputTensorDescs.Num(); Idx++)
+	{
+		TConstArrayView<int32> InputSymbolicTensorShapeData = InputTensorDescs[Idx].GetShape().GetData();
+
+		const FIntVector4 ModelInputShape = { 1, InputSymbolicTensorShapeData[1], Extent.Y, Extent.X };
 
 		InputShapes.Add(UE::NNE::FTensorShape::Make({ (uint32)ModelInputShape.X, (uint32)ModelInputShape.Y, (uint32)ModelInputShape.Z, (uint32)ModelInputShape.W }));
 
@@ -196,6 +236,11 @@ bool FInputProcessBase::PrepareAndValidate(IModelInstance& ModelInstance, FIntPo
 		return false;
 	}
 
+	if (!Validate(ModelInstance, Extent))
+	{
+		return false;
+	}
+
 	return true;
 }
 
@@ -268,7 +313,7 @@ bool FOutputProcessBase::Validate(const IModelInstance& ModelInstance, FIntPoint
 	if (OutputTensorDescs.Num() != OutputLayout.Num())
 	{
 		UE_LOG(LogNNEDenoiser, Error, TEXT("Wrong number of outputs (expected %d, got %d)!"), OutputLayout.Num(), OutputTensorDescs.Num())
-		return true;
+		return false;
 	}
 
 	for (int32 Idx = 0; Idx < OutputTensorDescs.Num(); Idx++)
@@ -283,9 +328,15 @@ bool FOutputProcessBase::Validate(const IModelInstance& ModelInstance, FIntPoint
 	}
 
 	TConstArrayView<UE::NNE::FTensorShape> OutputShapes = ModelInstance.GetOutputTensorShapes();
-	if (OutputShapes.Num() != OutputLayout.Num())
+	if (!OutputShapes.IsEmpty() && OutputShapes.Num() != OutputLayout.Num())
+	{
+		UE_LOG(LogNNEDenoiser, Error, TEXT("Wrong number of output shapes!"));
+		return false;
+	}
+
+	if (OutputShapes.IsEmpty() && OutputLayout.Num() > 0)
 	{
-		UE_LOG(LogNNEDenoiser, Error, TEXT("Wrong number of output shapes or not resolved yet!"));
+		UE_LOG(LogNNEDenoiser, Log, TEXT("Output shapes not resolved yet"));
 		return true;
 	}
 	
diff --git a/Engine/Plugins/Experimental/NNEDenoiser/Source/NNEDenoiser/Private/NNEDenoiserIOProcessBase.h b/Engine/Plugins/Experimental/NNEDenoiser/Source/NNEDenoiser/Private/NNEDenoiserIOProcessBase.h
--- a/Engine/Plugins/Experimental/NNEDenoiser/Source/NNEDenoiser/Private/NNEDenoiserIOProcessBase.h
+++ b/Engine/Plugins/Experimental/NNEDenoiser/Source/NNEDenoiser/Private/NNEDenoiserIOProcessBase.h
@@ -18,7 +18,9 @@ public:
 
 	virtual ~FInputProcessBase() = default;
 
-	virtual bool PrepareAndValidate(IModelInstance& ModelInstance, FIntPoint Extent) const override;
+	virtual bool Validate(const IModelInstance& ModelInstance, FIntPoint Extent) const override;
+
+	virtual bool Prepare(IModelInstance& ModelInstance, FIntPoint Extent) const override;
 
 	virtual int32 NumFrames(EResourceName Name) const override;
 
diff --git a/Engine/Plugins/Experimental/NNEDenoiser/Source/NNEDenoiser/Private/NNEDenoiserParameters.h b/Engine/Plugins/Experimental/NNEDenoiser/Source/NNEDenoiser/Private/NNEDenoiserParameters.h
--- a/Engine/Plugins/Experimental/NNEDenoiser/Source/NNEDenoiser/Private/NNEDenoiserParameters.h
+++ b/Engine/Plugins/Experimental/NNEDenoiser/Source/NNEDenoiser/Private/NNEDenoiserParameters.h
@@ -9,7 +9,15 @@ namespace UE::NNEDenoiser::Private
 	
 	struct FParameters
 	{
-		int32 TileMinimumOverlap = 20;
+		struct FTilingConfig
+		{
+			int32 Alignment = 1;
+			int32 Overlap{};
+			int32 MaxSize{};
+			int32 MinSize = 1;
+		};
+
+		FTilingConfig TilingConfig{};
 	};
 
 } // namespace UE::NNEDenoiser::Private
\ No newline at end of file
diff --git a/Engine/Plugins/Experimental/NNEDenoiser/Source/NNEDenoiser/Private/NNEDenoiserResourceManager.cpp b/Engine/Plugins/Experimental/NNEDenoiser/Source/NNEDenoiser/Private/NNEDenoiserResourceManager.cpp
--- a/Engine/Plugins/Experimental/NNEDenoiser/Source/NNEDenoiser/Private/NNEDenoiserResourceManager.cpp
+++ b/Engine/Plugins/Experimental/NNEDenoiser/Source/NNEDenoiser/Private/NNEDenoiserResourceManager.cpp
@@ -78,8 +78,9 @@ void FResourceManager::BeginTile(int32 TileIndex)
 
 	// Copy sub region from original input texture to tile input texture
 	FRHICopyTextureInfo CopyInputInfo{};
-	CopyInputInfo.SourcePosition = {Tile.Position.X, Tile.Position.Y, 0};
-	CopyInputInfo.Size = {Tiling.TileSize.X, Tiling.TileSize.Y, 1};
+	CopyInputInfo.SourcePosition = {Tile.Position.X + Tile.InputOffsets.Min.X, Tile.Position.Y + Tile.InputOffsets.Min.Y, 0};
+	CopyInputInfo.DestPosition = {Tile.InputOffsets.Min.X, Tile.InputOffsets.Min.Y, 0};
+	CopyInputInfo.Size = {Tiling.TileSize.X + Tile.InputOffsets.Width(), Tiling.TileSize.Y + Tile.InputOffsets.Height(), 1};
 
 	for (const auto& KeyValue : TextureMap)
 	{
diff --git a/Engine/Plugins/Experimental/NNEDenoiser/Source/NNEDenoiser/Private/NNEDenoiserSettings.cpp b/Engine/Plugins/Experimental/NNEDenoiser/Source/NNEDenoiser/Private/NNEDenoiserSettings.cpp
--- a/Engine/Plugins/Experimental/NNEDenoiser/Source/NNEDenoiser/Private/NNEDenoiserSettings.cpp
+++ b/Engine/Plugins/Experimental/NNEDenoiser/Source/NNEDenoiser/Private/NNEDenoiserSettings.cpp
@@ -11,7 +11,7 @@ UNNEDenoiserSettings::UNNEDenoiserSettings()
 	CategoryName = TEXT("Plugins");
 	SectionName = TEXT("NNE Denoiser");
 
-	DenoiserModelData = FSoftObjectPath(TEXT("/NNEDenoiser/NNEDNN_Oidn2_ColorAlbedoNormal_720.NNEDNN_Oidn2_ColorAlbedoNormal_720"));
+	DenoiserModelData = FSoftObjectPath(TEXT("/NNEDenoiser/NNEDNN_Oidn2_ColorAlbedoNormal_Dyn_Float16.NNEDNN_Oidn2_ColorAlbedoNormal_Dyn_Float16"));
 }
 
 void UNNEDenoiserSettings::PostInitProperties()
diff --git a/Engine/Plugins/Experimental/NNEDenoiser/Source/NNEDenoiser/Private/NNEDenoiserTiling.h b/Engine/Plugins/Experimental/NNEDenoiser/Source/NNEDenoiser/Private/NNEDenoiserTiling.h
--- a/Engine/Plugins/Experimental/NNEDenoiser/Source/NNEDenoiser/Private/NNEDenoiserTiling.h
+++ b/Engine/Plugins/Experimental/NNEDenoiser/Source/NNEDenoiser/Private/NNEDenoiserTiling.h
@@ -6,28 +6,56 @@
 #include "Math/IntRect.h"
 #include "Math/IntVector.h"
 #include "Math/MathFwd.h"
+#include "NNEDenoiserUtils.h"
 
 namespace UE::NNEDenoiser::Private
 {
 	
 	struct FTile
 	{
+		FTile(FIntPoint Position, FIntRect InputOffsets, FIntRect OutputOffsets) : Position(Position), InputOffsets(InputOffsets), OutputOffsets(OutputOffsets) {}
+		
 		FIntPoint Position;
+		FIntRect InputOffsets;
 		FIntRect OutputOffsets;
 	};
 
 	struct FTiling
 	{
+		FTiling() = default;
+
 		FIntPoint TileSize;
 		FIntPoint Count;
 		TArray<FTile> Tiles;
 	};
 
-	inline FTiling CreateTiling(FIntPoint TargetTileSize, FIntPoint MinimumOverlap, FIntPoint Size)
+	inline FTiling CreateTiling(FIntPoint TargetTileSize, FIntPoint MaxTileSize, FIntPoint MinTileSize, int32 TileAlignment, FIntPoint MinimumOverlap, FIntPoint Size)
 	{
+		auto GetTileSize = [] (int32 Max, int32 Min, int32 Alignment, int32 Overlap, int32 Size)
+		{
+			int NumTiles = 1;
+			int32 Result = FMath::Max(RoundUp(Size, Alignment), Min);
+
+			if (Max <= 0)
+			{
+				return Result;
+			}
+
+			Alignment = FMath::Max(Alignment, 1);
+			Min = FMath::Max(Min, Alignment + Overlap);
+
+			while (Result > Max && Result > Min)
+			{
+				NumTiles++;
+				Result = FMath::Max(RoundUp(CeilDiv(Size - Overlap, NumTiles), Alignment) + Overlap, Min);
+			}
+
+			return Result;
+		};
+
 		auto GetNumTiles = [] (int32 TileSize, int32 Overlap, int32 Size)
 		{
-			return FMath::Max(FMath::CeilToInt32((Size - TileSize) / (float)(TileSize - Overlap)), 0) + 1;
+			return Size > TileSize ? CeilDiv(Size - Overlap, TileSize - Overlap) : 1;
 		};
 
 		auto GetOffsets = [] (int32 Count, int32 TileSize, int32 Overlap, int32 Size)
@@ -35,7 +63,7 @@ namespace UE::NNEDenoiser::Private
 			TArray<int32> Result;
 			for (int32 I = 0; I < Count; I++)
 			{
-				Result.Add(I < Count - 1 ? I * (TileSize - Overlap) : Size - TileSize);
+				Result.Add(I < Count - 1 ? I * (TileSize - Overlap) : FMath::Max(Size - TileSize, 0));
 			}
 
 			return Result;
@@ -43,8 +71,8 @@ namespace UE::NNEDenoiser::Private
 
 		FTiling Result{};
 		Result.TileSize = {
-			FMath::Min(TargetTileSize.X, Size.X),
-			FMath::Min(TargetTileSize.Y, Size.Y)
+			TargetTileSize.X > 0 ? TargetTileSize.X : GetTileSize(MaxTileSize.X, MinTileSize.X, TileAlignment, MinimumOverlap.X, Size.X),
+			TargetTileSize.Y > 0 ? TargetTileSize.Y : GetTileSize(MaxTileSize.Y, MinTileSize.Y, TileAlignment, MinimumOverlap.Y, Size.Y)
 		};
 
 		Result.Count = {
@@ -66,26 +94,28 @@ namespace UE::NNEDenoiser::Private
 		const TArray<int32> OffsetsX = GetOffsets(Result.Count.X, Result.TileSize.X, Overlap.X, Size.X);
 		const TArray<int32> OffsetsY = GetOffsets(Result.Count.Y, Result.TileSize.Y, Overlap.Y, Size.Y);
 
+		const int32 InX0 = 0;
+		const int32 InX1 = -FMath::Max(Result.TileSize.X - Size.X, 0);
+		const int32 InY0 = 0;
+		const int32 InY1 = -FMath::Max(Result.TileSize.Y - Size.Y, 0);
+
 		for (int32 Ty = 0; Ty < Result.Count.Y; Ty++)
 		{
 			const int32 Y0 = OffsetsY[Ty];
 			const int32 Y1 = Y0 + Result.TileSize.Y;
 
-			const int32 OutY0 = Ty > 0 ? HalfOverlap.Y : 0;
-			const int32 OutY1 = Ty < Result.Count.Y - 1 ? -HalfOverlap.Y : 0;
+			const int32 OutY0 = Result.Count.Y == 1 ? InY0 : Ty > 0 ? HalfOverlap.Y : 0;
+			const int32 OutY1 = Result.Count.Y == 1 ? InY1 : Ty < Result.Count.Y - 1 ? -HalfOverlap.Y : 0;
 
 			for (int32 Tx = 0; Tx < Result.Count.X; Tx++)
 			{
 				const int32 X0 = OffsetsX[Tx];
 				const int32 X1 = X0 + Result.TileSize.X;
 
-				const int32 OutX0 = Tx > 0 ? HalfOverlap.X : 0;
-				const int32 OutX1 = Tx < Result.Count.X - 1 ? -HalfOverlap.X : 0;
-
-				const FIntPoint InputPosition = {X0, Y0};
-				const FIntRect OutputOffsets = {OutX0, OutY0, OutX1, OutY1};
+				const int32 OutX0 = Result.Count.X == 1 ? InX0 : Tx > 0 ? HalfOverlap.X : 0;
+				const int32 OutX1 = Result.Count.X == 1 ? InX1 : Tx < Result.Count.X - 1 ? -HalfOverlap.X : 0;
 
-				Result.Tiles.Add(FTile{InputPosition, OutputOffsets});
+				Result.Tiles.Add(FTile{{X0, Y0}, {InX0, InY0, InX1, InY1}, {OutX0, OutY0, OutX1, OutY1}});
 			}
 		}
 
diff --git a/Engine/Plugins/Experimental/NNEDenoiser/Source/NNEDenoiser/Private/NNEDenoiserUtils.h b/Engine/Plugins/Experimental/NNEDenoiser/Source/NNEDenoiser/Private/NNEDenoiserUtils.h
--- a/Engine/Plugins/Experimental/NNEDenoiser/Source/NNEDenoiser/Private/NNEDenoiserUtils.h
+++ b/Engine/Plugins/Experimental/NNEDenoiser/Source/NNEDenoiser/Private/NNEDenoiserUtils.h
@@ -7,7 +7,6 @@
 #include "NNEDenoiserParameters.h"
 #include "NNEDenoiserShadersDefaultCS.h"
 #include "NNEDenoiserShadersOidnCS.h"
-#include "NNEDenoiserTiling.h"
 #include "NNETypes.h"
 #include "RHICommandList.h"
 #include "RHIResources.h"
@@ -15,6 +14,19 @@
 
 namespace UE::NNEDenoiser::Private
 {
+	// Returns ceil(a / b) for non-negative integers
+	template<typename Int, typename IntB>
+	Int CeilDiv(Int a, IntB b)
+	{
+		return (a + b - 1) / b;
+	}
+
+	// Returns a rounded up to multiple of b
+	template<typename Int, typename IntB>
+	Int RoundUp(Int a, IntB b)
+	{
+		return CeilDiv(a, b) * b;
+	}
 
 	template<class IntType>
 	bool IsTensorShapeValid(TConstArrayView<IntType> ShapeData, TConstArrayView<int32> RequiredShapeData, const FString& Label)
@@ -31,7 +43,7 @@ namespace UE::NNEDenoiser::Private
 		{
 			if (RequiredShapeData[I] >= 0 && (int32)ShapeData[I] != RequiredShapeData[I])
 			{
-				UE_LOG(LogNNEDenoiser, Error, TEXT("%s does not have required shape!"), *Label)
+				UE_LOG(LogNNEDenoiser, Error, TEXT("%s does not have required shape (expected %d, got %d)!"), *Label, RequiredShapeData[I], ShapeData[I])
 				return false;
 			}
 		}
diff --git a/Engine/Plugins/Experimental/NNEDenoiser/Source/NNEDenoiser/Private/NNEDenoiserViewExtension.cpp b/Engine/Plugins/Experimental/NNEDenoiser/Source/NNEDenoiser/Private/NNEDenoiserViewExtension.cpp
--- a/Engine/Plugins/Experimental/NNEDenoiser/Source/NNEDenoiser/Private/NNEDenoiserViewExtension.cpp
+++ b/Engine/Plugins/Experimental/NNEDenoiser/Source/NNEDenoiser/Private/NNEDenoiserViewExtension.cpp
@@ -13,6 +13,7 @@
 #include "NNEDenoiserModelIOMappingData.h"
 #include "NNEDenoiserPathTracingDenoiser.h"
 #include "NNEDenoiserPathTracingSpatialTemporalDenoiser.h"
+#include "NNEDenoiserUtils.h"
 #include "NNEModelData.h"
 #include "PathTracingDenoiser.h"
 
@@ -33,9 +34,9 @@ static TAutoConsoleVariable<int32> CVarNNEDenoiserModelData(
 	TEXT("  1: OIDN Color\n")
 	TEXT("  2: OIDN Color and Albedo\n")
 	TEXT("  3: OIDN Color, Albedo and Normal\n")
-	TEXT("  4: OIDN Color, Half\n")
-	TEXT("  5: OIDN Color and Albedo, Half\n")
-	TEXT("  6: OIDN Color, Albedo and Normal, Half")
+	TEXT("  4: OIDN Color | Float16\n")
+	TEXT("  5: OIDN Color and Albedo | Float16\n")
+	TEXT("  6: OIDN Color, Albedo and Normal | Float16")
 );
 
 static TAutoConsoleVariable<int32> CVarNNEDenoiserRuntimeType(
@@ -67,12 +68,12 @@ FString GetDenoiserModelDataNameFromCVarAndSettings(const UNNEDenoiserSettings*
 	switch(Idx)
 	{
 		case 0: return !Settings->DenoiserModelData.IsNull() ? Settings->DenoiserModelData.ToString() : FString();
-		case 1: return TEXT("/NNEDenoiser/NNEDNN_Oidn2_Color_720.NNEDNN_Oidn2_Color_720");
-		case 2: return TEXT("/NNEDenoiser/NNEDNN_Oidn2_ColorAlbedo_720.NNEDNN_Oidn2_ColorAlbedo_720");
-		case 3: return TEXT("/NNEDenoiser/NNEDNN_Oidn2_ColorAlbedoNormal_720.NNEDNN_Oidn2_ColorAlbedoNormal_720");
-		case 4: return TEXT("/NNEDenoiser/NNEDNN_Oidn2_Color_720_Float16.NNEDNN_Oidn2_Color_720_Float16");
-		case 5: return TEXT("/NNEDenoiser/NNEDNN_Oidn2_ColorAlbedo_720_Float16.NNEDNN_Oidn2_ColorAlbedo_720_Float16");
-		case 6: return TEXT("/NNEDenoiser/NNEDNN_Oidn2_ColorAlbedoNormal_720_Float16.NNEDNN_Oidn2_ColorAlbedoNormal_720_Float16");
+		case 1: return TEXT("/NNEDenoiser/NNEDNN_Oidn2_Color_Dyn.NNEDNN_Oidn2_Color_Dyn");
+		case 2: return TEXT("/NNEDenoiser/NNEDNN_Oidn2_ColorAlbedo_Dyn.NNEDNN_Oidn2_ColorAlbedo_Dyn");
+		case 3: return TEXT("/NNEDenoiser/NNEDNN_Oidn2_ColorAlbedoNormal_Dyn.NNEDNN_Oidn2_ColorAlbedoNormal_Dyn");
+		case 4: return TEXT("/NNEDenoiser/NNEDNN_Oidn2_Color_Dyn_Float16.NNEDNN_Oidn2_Color_Dyn_Float16");
+		case 5: return TEXT("/NNEDenoiser/NNEDNN_Oidn2_ColorAlbedo_Dyn_Float16.NNEDNN_Oidn2_ColorAlbedo_Dyn_Float16");
+		case 6: return TEXT("/NNEDenoiser/NNEDNN_Oidn2_ColorAlbedoNormal_Dyn_Float16.NNEDNN_Oidn2_ColorAlbedoNormal_Dyn_Float16");
 	}
 	check(false);
 	return FString();
@@ -142,6 +143,57 @@ static FResourceMappingList MakeTensorLayout(UDataTable* DataTable)
 	return Result;
 }
 
+FParameters GetParametersValidated(const UNNEDenoiserModelData& DenoiserModelData)
+{
+	FParameters Parameters{
+		.TilingConfig =
+		{
+			.Alignment = DenoiserModelData.TilingConfig.Alignment,
+			.Overlap = DenoiserModelData.TilingConfig.Overlap,
+			.MaxSize = DenoiserModelData.TilingConfig.MaxSize,
+			.MinSize = DenoiserModelData.TilingConfig.MinSize
+		}
+	};
+
+	if (Parameters.TilingConfig.Alignment < 1)
+	{
+		UE_LOG(LogNNEDenoiser, Warning, TEXT("Tiling alignment should be at least 1!"));
+		Parameters.TilingConfig.Alignment = 1;
+	}
+
+	if (Parameters.TilingConfig.Overlap % Parameters.TilingConfig.Alignment != 0)
+	{
+		UE_LOG(LogNNEDenoiser, Warning, TEXT("Tiling overlap should be aligned by %d!"), Parameters.TilingConfig.Alignment);
+		Parameters.TilingConfig.Overlap = RoundUp(Parameters.TilingConfig.Overlap, Parameters.TilingConfig.Alignment);
+	}
+
+	if (Parameters.TilingConfig.MinSize < Parameters.TilingConfig.Overlap + Parameters.TilingConfig.Alignment)
+	{
+		UE_LOG(LogNNEDenoiser, Warning, TEXT("Minimum tile size should be at least overlap + alignment = %d!"), Parameters.TilingConfig.Overlap + Parameters.TilingConfig.Alignment);
+		Parameters.TilingConfig.MinSize = Parameters.TilingConfig.Overlap + Parameters.TilingConfig.Alignment;
+	}
+
+	if (Parameters.TilingConfig.MinSize % Parameters.TilingConfig.Alignment != 0)
+	{
+		UE_LOG(LogNNEDenoiser, Warning, TEXT("Minimum tile size should be aligned by %d!"), Parameters.TilingConfig.Alignment);
+		Parameters.TilingConfig.MinSize = RoundUp(Parameters.TilingConfig.MinSize, Parameters.TilingConfig.Alignment);
+	}
+
+	if (Parameters.TilingConfig.MaxSize > 0 && Parameters.TilingConfig.MaxSize < Parameters.TilingConfig.MinSize)
+	{
+		UE_LOG(LogNNEDenoiser, Warning, TEXT("Maximum tile size should be at least minimum tile size!"));
+		Parameters.TilingConfig.MaxSize = Parameters.TilingConfig.MinSize;
+	}
+
+	if (Parameters.TilingConfig.MaxSize % Parameters.TilingConfig.Alignment != 0)
+	{
+		UE_LOG(LogNNEDenoiser, Warning, TEXT("Maximum tile size should be aligned by %d!"), Parameters.TilingConfig.Alignment);
+		Parameters.TilingConfig.MaxSize = RoundUp(Parameters.TilingConfig.MaxSize, Parameters.TilingConfig.Alignment);
+	}
+
+	return Parameters;
+}
+
 TUniquePtr<FGenericDenoiser> CreateNNEDenoiserFromAsset(const FString& AssetName, EDenoiserRuntimeType RuntimeType, const FString& RuntimeNameOverride)
 {
 	if (AssetName.IsEmpty())
@@ -211,7 +263,9 @@ TUniquePtr<FGenericDenoiser> CreateNNEDenoiserFromAsset(const FString& AssetName
 		OutputProcess = MakeUnique<FOutputProcessBase>(MoveTemp(OutputLayout));
 	}
 
-	return CreateNNEDenoiser(*ModelData, RuntimeType, RuntimeNameOverride, MoveTemp(InputProcess), MoveTemp(OutputProcess), {});
+	FParameters Parameters = GetParametersValidated(*DenoiserModelData);
+
+	return CreateNNEDenoiser(*ModelData, RuntimeType, RuntimeNameOverride, MoveTemp(InputProcess), MoveTemp(OutputProcess), MoveTemp(Parameters));
 }
 
 FViewExtension::FViewExtension(const FAutoRegister& AutoRegister) : FSceneViewExtensionBase(AutoRegister)
diff --git a/Engine/Plugins/Experimental/NNEDenoiser/Source/NNEDenoiser/Public/NNEDenoiserModelData.h b/Engine/Plugins/Experimental/NNEDenoiser/Source/NNEDenoiser/Public/NNEDenoiserModelData.h
--- a/Engine/Plugins/Experimental/NNEDenoiser/Source/NNEDenoiser/Public/NNEDenoiserModelData.h
+++ b/Engine/Plugins/Experimental/NNEDenoiser/Source/NNEDenoiser/Public/NNEDenoiserModelData.h
@@ -8,6 +8,29 @@
 
 #include "NNEDenoiserModelData.generated.h"
 
+/** Tiling configuration for fixed and dynamic size models */
+USTRUCT(BlueprintType)
+struct FTilingConfig
+{
+	GENERATED_BODY()
+
+	/** Tile size alignment (applies only to dynamic size models) */
+	UPROPERTY(EditAnywhere, BlueprintReadWrite, Category=NNEDenoiser)
+	int32 Alignment = 1;
+
+	/** Tile overlap */
+	UPROPERTY(EditAnywhere, BlueprintReadWrite, Category=NNEDenoiser)
+	int32 Overlap = 0;
+
+	/** Maximum tile size (applies only to dynamic size models) */
+	UPROPERTY(EditAnywhere, BlueprintReadWrite, Category=NNEDenoiser)
+	int32 MaxSize = 0;
+
+	/** Minimum tile size (applies only to dynamic size models) */
+	UPROPERTY(EditAnywhere, BlueprintReadWrite, Category=NNEDenoiser)
+	int32 MinSize = 1;
+};
+
 /** Denoiser model data asset */
 UCLASS(BlueprintType)
 class NNEDENOISER_API UNNEDenoiserModelData : public UDataAsset
@@ -26,4 +49,8 @@ public:
 	/** Output mapping table */
 	UPROPERTY(EditAnywhere, BlueprintReadWrite, Category=NNEDenoiser)
 	TSoftObjectPtr<UDataTable> OutputMapping;
+
+	/** Tiling configuration */
+	UPROPERTY(EditAnywhere, BlueprintReadWrite, Category=NNEDenoiser)
+	FTilingConfig TilingConfig{};
 };
\ No newline at end of file
diff --git a/Engine/Plugins/Experimental/NNEDenoiser/Source/NNEDenoiserShaders/Shaders/NNEDenoiserShadersDefault.usf b/Engine/Plugins/Experimental/NNEDenoiser/Source/NNEDenoiserShaders/Shaders/NNEDenoiserShadersDefault.usf
--- a/Engine/Plugins/Experimental/NNEDenoiser/Source/NNEDenoiserShaders/Shaders/NNEDenoiserShadersDefault.usf
+++ b/Engine/Plugins/Experimental/NNEDenoiser/Source/NNEDenoiserShaders/Shaders/NNEDenoiserShadersDefault.usf
@@ -63,11 +63,7 @@ void WriteOutput(in const uint3 DispatchThreadID : SV_DispatchThreadID)
 		return;
 	}
 
-	// nearest neighbour
-	const int X = (int)(DispatchThreadID.x * (OutputBufferWidth - 1.0) / (OutputTextureWidth - 1.0));
-	const int Y = (int)(DispatchThreadID.y * (OutputBufferHeight - 1.0) / (OutputTextureHeight - 1.0));
-
-	const int Idx = Y * OutputBufferWidth + X;
+	const int Idx = DispatchThreadID.y * OutputBufferWidth + DispatchThreadID.x;
 	const int Offset = OutputBufferWidth * OutputBufferHeight;
 
 	for (int I = 0; I < NUM_MAPPED_CHANNELS; I++)
