[NNEDenoiser] Updated OIDN models to version 2.2.0.

#tests Editor, Win64 run all models

#rnx

#virtualized

diff --git a/Engine/Plugins/Experimental/NNEDenoiser/Source/NNEDenoiser/Private/NNEDenoiserViewExtension.cpp b/Engine/Plugins/Experimental/NNEDenoiser/Source/NNEDenoiser/Private/NNEDenoiserViewExtension.cpp
--- a/Engine/Plugins/Experimental/NNEDenoiser/Source/NNEDenoiser/Private/NNEDenoiserViewExtension.cpp
+++ b/Engine/Plugins/Experimental/NNEDenoiser/Source/NNEDenoiser/Private/NNEDenoiserViewExtension.cpp
@@ -68,12 +68,12 @@ FString GetDenoiserModelDataNameFromCVarAndSettings(const UNNEDenoiserSettings*
 	switch(Idx)
 	{
 		case 0: return !Settings->DenoiserModelData.IsNull() ? Settings->DenoiserModelData.ToString() : FString();
-		case 1: return TEXT("/NNEDenoiser/NNEDNN_Oidn2_Color_Dyn.NNEDNN_Oidn2_Color_Dyn");
-		case 2: return TEXT("/NNEDenoiser/NNEDNN_Oidn2_ColorAlbedo_Dyn.NNEDNN_Oidn2_ColorAlbedo_Dyn");
-		case 3: return TEXT("/NNEDenoiser/NNEDNN_Oidn2_ColorAlbedoNormal_Dyn.NNEDNN_Oidn2_ColorAlbedoNormal_Dyn");
-		case 4: return TEXT("/NNEDenoiser/NNEDNN_Oidn2_Color_Dyn_Float16.NNEDNN_Oidn2_Color_Dyn_Float16");
-		case 5: return TEXT("/NNEDenoiser/NNEDNN_Oidn2_ColorAlbedo_Dyn_Float16.NNEDNN_Oidn2_ColorAlbedo_Dyn_Float16");
-		case 6: return TEXT("/NNEDenoiser/NNEDNN_Oidn2_ColorAlbedoNormal_Dyn_Float16.NNEDNN_Oidn2_ColorAlbedoNormal_Dyn_Float16");
+		case 1: return TEXT("/NNEDenoiser/NNEDNN_Oidn2-2-0_Color_Dyn.NNEDNN_Oidn2-2-0_Color_Dyn");
+		case 2: return TEXT("/NNEDenoiser/NNEDNN_Oidn2-2-0_ColorAlbedo_Dyn.NNEDNN_Oidn2-2-0_ColorAlbedo_Dyn");
+		case 3: return TEXT("/NNEDenoiser/NNEDNN_Oidn2-2-0_ColorAlbedoNormal_Dyn.NNEDNN_Oidn2-2-0_ColorAlbedoNormal_Dyn");
+		case 4: return TEXT("/NNEDenoiser/NNEDNN_Oidn2-2-0_Color_Dyn_Float16.NNEDNN_Oidn2-2-0_Color_Dyn_Float16");
+		case 5: return TEXT("/NNEDenoiser/NNEDNN_Oidn2-2-0_ColorAlbedo_Dyn_Float16.NNEDNN_Oidn2-2-0_ColorAlbedo_Dyn_Float16");
+		case 6: return TEXT("/NNEDenoiser/NNEDNN_Oidn2-2-0_ColorAlbedoNormal_Dyn_Float16.NNEDNN_Oidn2-2-0_ColorAlbedoNormal_Dyn_Float16");
 	}
 	check(false);
 	return FString();
