[NNEDenoiser] Update default neural denoiser model.

#tests Editor, Win64 start up and verify default denoier model

#rnx

diff --git a/Engine/Plugins/Experimental/NNEDenoiser/Source/NNEDenoiser/Private/NNEDenoiserSettings.cpp b/Engine/Plugins/Experimental/NNEDenoiser/Source/NNEDenoiser/Private/NNEDenoiserSettings.cpp
--- a/Engine/Plugins/Experimental/NNEDenoiser/Source/NNEDenoiser/Private/NNEDenoiserSettings.cpp
+++ b/Engine/Plugins/Experimental/NNEDenoiser/Source/NNEDenoiser/Private/NNEDenoiserSettings.cpp
@@ -11,7 +11,7 @@ UNNEDenoiserSettings::UNNEDenoiserSettings()
 	CategoryName = TEXT("Plugins");
 	SectionName = TEXT("NNE Denoiser");
 
-	DenoiserModelData = FSoftObjectPath(TEXT("/NNEDenoiser/NNEDNN_Oidn2_ColorAlbedoNormal_Dyn_Float16.NNEDNN_Oidn2_ColorAlbedoNormal_Dyn_Float16"));
+	DenoiserModelData = FSoftObjectPath(TEXT("/NNEDenoiser/NNEDNN_Oidn2-2-0_ColorAlbedoNormal_Dyn_Float16.NNEDNN_Oidn2-2-0_ColorAlbedoNormal_Dyn_Float16"));
 }
 
 void UNNEDenoiserSettings::PostInitProperties()
