[NNEDenoiser] Clear intermediate textures to avoid artifacts when not using the whole texture area.

#tests Editor, Win64 move camera with 1spp and verify absence of artifacts ;)

#rnx

diff --git a/Engine/Plugins/Experimental/NNEDenoiser/Source/NNEDenoiser/Private/NNEDenoiserResourceManager.cpp b/Engine/Plugins/Experimental/NNEDenoiser/Source/NNEDenoiser/Private/NNEDenoiserResourceManager.cpp
--- a/Engine/Plugins/Experimental/NNEDenoiser/Source/NNEDenoiser/Private/NNEDenoiserResourceManager.cpp
+++ b/Engine/Plugins/Experimental/NNEDenoiser/Source/NNEDenoiser/Private/NNEDenoiserResourceManager.cpp
@@ -82,6 +82,10 @@ void FResourceManager::BeginTile(int32 TileIndex)
 	CopyInputInfo.DestPosition = {Tile.InputOffsets.Min.X, Tile.InputOffsets.Min.Y, 0};
 	CopyInputInfo.Size = {Tiling.TileSize.X + Tile.InputOffsets.Width(), Tiling.TileSize.Y + Tile.InputOffsets.Height(), 1};
 
+	// Only clear if we don't overwrite the whole texture
+	const bool bNeedClear = Tile.InputOffsets.Width() != 0 || Tile.InputOffsets.Height() != 0;
+	const FLinearColor ClearColor = {.0f, 0.f, .0f, .0f};
+
 	for (const auto& KeyValue : TextureMap)
 	{
 		if (!IntermediateTextureMap.Contains(KeyValue.Key))
@@ -96,6 +100,10 @@ void FResourceManager::BeginTile(int32 TileIndex)
 		const int32 StartIdx = KeyValue.Key == EResourceName::Output ? 1 : 0;
 		for (int32 I = StartIdx; I < Textures.Num(); I++)
 		{
+			if (bNeedClear)
+			{
+				AddClearRenderTargetPass(GraphBuilder, IntermediateTextures[I], ClearColor);
+			}
 			AddCopyTexturePass(GraphBuilder, Textures[I], IntermediateTextures[I], CopyInputInfo);
 		}
 	}
