[NNEDenoiser] Enable support for denoising alpha. Added OIDN models with alpha denoising. Removed OIDN model using only Albedo texture, but added high quality model that prefilters Albedo and Normal before denoising beauty image (OIDN cleanAux parameter).

#tests Editor, Win64 run all models

#rnx

#virtualized

diff --git a/Engine/Plugins/Experimental/NNEDenoiser/Source/NNEDenoiser/Private/NNEDenoiserGenericDenoiser.cpp b/Engine/Plugins/Experimental/NNEDenoiser/Source/NNEDenoiser/Private/NNEDenoiserGenericDenoiser.cpp
--- a/Engine/Plugins/Experimental/NNEDenoiser/Source/NNEDenoiser/Private/NNEDenoiserGenericDenoiser.cpp
+++ b/Engine/Plugins/Experimental/NNEDenoiser/Source/NNEDenoiser/Private/NNEDenoiserGenericDenoiser.cpp
@@ -173,6 +173,9 @@ bool FGenericDenoiser::Prepare(FIntPoint Extent)
 
 	LastExtent = Extent;
 
+	UE_LOG(LogNNEDenoiser, Log, TEXT("Prepared neural denoiser model:\n  Viewport   %dx%d\n  Num. tiles %dx%d\n  Tile size  %dx%d"),
+		Extent.X, Extent.Y, Tiling.Count.X, Tiling.Count.Y, Tiling.TileSize.X, Tiling.TileSize.Y);
+
 	return true;
 }
 
@@ -226,8 +229,6 @@ TUniquePtr<FHistory> FGenericDenoiser::AddPasses(
 	TArray<FRDGBufferRef> InputBuffers = CreateBuffersRDG(GraphBuilder, ModelInstance->GetInputTensorDescs(), ModelInstance->GetInputTensorShapes());
 	TArray<FRDGBufferRef> OutputBuffers = CreateBuffersRDG(GraphBuilder, ModelInstance->GetOutputTensorDescs(), GetOutputTensorShapes(*ModelInstance));
 
-	UE_LOG(LogNNEDenoiser, Log, TEXT("Divided work of size %dx%d into %dx%d tiles of size %dx%d each..."), Extent.X, Extent.Y, Tiling.Count.X, Tiling.Count.Y, Tiling.TileSize.X, Tiling.TileSize.Y);
-
 	for (int32 I = 0; I < Tiling.Tiles.Num(); I++)
 	{
 		ResourceManager.BeginTile(I);
diff --git a/Engine/Plugins/Experimental/NNEDenoiser/Source/NNEDenoiser/Private/NNEDenoiserIOProcessBase.cpp b/Engine/Plugins/Experimental/NNEDenoiser/Source/NNEDenoiser/Private/NNEDenoiserIOProcessBase.cpp
--- a/Engine/Plugins/Experimental/NNEDenoiser/Source/NNEDenoiser/Private/NNEDenoiserIOProcessBase.cpp
+++ b/Engine/Plugins/Experimental/NNEDenoiser/Source/NNEDenoiser/Private/NNEDenoiserIOProcessBase.cpp
@@ -41,21 +41,19 @@ void AddReadInputPass(
 
 	const FIntVector InputTextureSize = InputTexture->Desc.GetSize();
 
-	FNNEDenoiserReadInputCS::FParameters *ReadInputParameters = GraphBuilder.AllocParameters<FNNEDenoiserReadInputCS::FParameters>();
-	ReadInputParameters->InputTextureWidth = InputTextureSize.X;
-	ReadInputParameters->InputTextureHeight = InputTextureSize.Y;
+	FNNEDenoiserTextureBufferMappedCopyCS::FParameters *ReadInputParameters = GraphBuilder.AllocParameters<FNNEDenoiserTextureBufferMappedCopyCS::FParameters>();
+	ReadInputParameters->Width = InputBufferSize.X;
+	ReadInputParameters->Height = InputBufferSize.Y;
 	ReadInputParameters->InputTexture = InputTexture;
-	ReadInputParameters->InputBufferWidth = InputBufferSize.X;
-	ReadInputParameters->InputBufferHeight = InputBufferSize.Y;
-	ReadInputParameters->InputBuffer = InputBufferUAV;
+	ReadInputParameters->OutputBuffer = InputBufferUAV;
 	for (int32 Idx = 0; Idx < ChannelMapping.Num(); Idx++)
 	{
-		ReadInputParameters->BufferChannel_TextureChannel_Unused_Unused[Idx] = { ChannelMapping[Idx].X, ChannelMapping[Idx].Y, 0, 0 };
+		ReadInputParameters->OutputChannel_InputChannel_Unused_Unused[Idx] = { ChannelMapping[Idx].X, ChannelMapping[Idx].Y, 0, 0 };
 	}
 
-	FNNEDenoiserReadInputCS::FPermutationDomain PermutationVector;
-	PermutationVector.Set<FNNEDenoiserReadInputCS::FNNEDenoiserDataType>(DataType);
-	PermutationVector.Set<FNNEDenoiserReadInputCS::FNNEDenoiserNumMappedChannels>(ChannelMapping.Num());
+	FNNEDenoiserTextureBufferMappedCopyCS::FPermutationDomain PermutationVector;
+	PermutationVector.Set<FNNEDenoiserTextureBufferMappedCopyCS::FNNEDenoiserDataType>(DataType);
+	PermutationVector.Set<FNNEDenoiserTextureBufferMappedCopyCS::FNNEDenoiserNumMappedChannels>(ChannelMapping.Num());
 
 	FIntVector ReadInputThreadGroupCount = FIntVector(
 		FMath::DivideAndRoundUp(InputBufferSize.X, FNNEDenoiserConstants::THREAD_GROUP_SIZE),
@@ -63,7 +61,7 @@ void AddReadInputPass(
 		1);
 
 	FGlobalShaderMap* GlobalShaderMap = GetGlobalShaderMap(GMaxRHIFeatureLevel);
-	TShaderMapRef<FNNEDenoiserReadInputCS> ReadInputShader(GlobalShaderMap, PermutationVector);
+	TShaderMapRef<FNNEDenoiserTextureBufferMappedCopyCS> ReadInputShader(GlobalShaderMap, PermutationVector);
 
 	RDG_EVENT_SCOPE(GraphBuilder, "NNEDenoiser.ReadInput");
 	RDG_GPU_STAT_SCOPE(GraphBuilder, FNNEDenoiserReadInput);
@@ -89,21 +87,19 @@ void AddWriteOutputPass(
 
 	const FIntVector OutputTextureSize = OutputTextureUAV->Desc.Texture->Desc.GetSize();
 
-	FNNEDenoiserWriteOutputCS::FParameters *WriteOutputParameters = GraphBuilder.AllocParameters<FNNEDenoiserWriteOutputCS::FParameters>();
-	WriteOutputParameters->OutputBufferWidth = OutputBufferSize.X;
-	WriteOutputParameters->OutputBufferHeight = OutputBufferSize.Y;
-	WriteOutputParameters->OutputBuffer = OutputBufferUAV;
-	WriteOutputParameters->OutputTextureWidth = OutputTextureSize.X;
-	WriteOutputParameters->OutputTextureHeight = OutputTextureSize.Y;
+	FNNEDenoiserBufferTextureMappedCopyCS::FParameters *WriteOutputParameters = GraphBuilder.AllocParameters<FNNEDenoiserBufferTextureMappedCopyCS::FParameters>();
+	WriteOutputParameters->Width = OutputBufferSize.X;
+	WriteOutputParameters->Height = OutputBufferSize.Y;
+	WriteOutputParameters->InputBuffer = OutputBufferUAV;
 	WriteOutputParameters->OutputTexture = OutputTextureUAV;
 	for (int32 Idx = 0; Idx < ChannelMapping.Num(); Idx++)
 	{
-		WriteOutputParameters->BufferChannel_TextureChannel_Unused_Unused[Idx] = { ChannelMapping[Idx].X, ChannelMapping[Idx].Y, 0, 0 };
+		WriteOutputParameters->OutputChannel_InputChannel_Unused_Unused[Idx] = { ChannelMapping[Idx].Y, ChannelMapping[Idx].X, 0, 0 };
 	}
 
-	FNNEDenoiserWriteOutputCS::FPermutationDomain PermutationVector;
-	PermutationVector.Set<FNNEDenoiserWriteOutputCS::FNNEDenoiserDataType>(DataType);
-	PermutationVector.Set<FNNEDenoiserWriteOutputCS::FNNEDenoiserNumMappedChannels>(ChannelMapping.Num());
+	FNNEDenoiserBufferTextureMappedCopyCS::FPermutationDomain PermutationVector;
+	PermutationVector.Set<FNNEDenoiserBufferTextureMappedCopyCS::FNNEDenoiserDataType>(DataType);
+	PermutationVector.Set<FNNEDenoiserBufferTextureMappedCopyCS::FNNEDenoiserNumMappedChannels>(ChannelMapping.Num());
 
 	FIntVector WriteOutputThreadGroupCount = FIntVector(
 		FMath::DivideAndRoundUp(OutputTextureSize.X, FNNEDenoiserConstants::THREAD_GROUP_SIZE),
@@ -111,7 +107,7 @@ void AddWriteOutputPass(
 		1);
 
 	FGlobalShaderMap* GlobalShaderMap = GetGlobalShaderMap(GMaxRHIFeatureLevel);
-	TShaderMapRef<FNNEDenoiserWriteOutputCS> WriteOutputShader(GlobalShaderMap, PermutationVector);
+	TShaderMapRef<FNNEDenoiserBufferTextureMappedCopyCS> WriteOutputShader(GlobalShaderMap, PermutationVector);
 
 	RDG_EVENT_SCOPE(GraphBuilder, "NNEDenoiser.WriteOutput");
 	RDG_GPU_STAT_SCOPE(GraphBuilder, FNNEDenoiserWriteOutput);
@@ -368,7 +364,12 @@ void FOutputProcessBase::AddPasses(
 {
 	for (int32 I = 0; I < TensorDescs.Num(); I++)
 	{
-		ReadOutputBuffer(GraphBuilder, TensorDescs[I], TensorShapes[I], ResourceAccess, Buffers[I], OutputLayout.GetChecked(I), OutputTexture);
+		// outputs might be discarded
+		const FResourceMapping* ResourceMapping = OutputLayout.Get(I);
+		if (ResourceMapping && ResourceMapping->GetChannelMappingPerFrame(EResourceName::Output).Num() == 1)
+		{
+			ReadOutputBuffer(GraphBuilder, TensorDescs[I], TensorShapes[I], ResourceAccess, Buffers[I], *ResourceMapping, OutputTexture);
+		}
 	}
 
 	const EResourceName TensorName = EResourceName::Output;
diff --git a/Engine/Plugins/Experimental/NNEDenoiser/Source/NNEDenoiser/Private/NNEDenoiserResourceMapping.h b/Engine/Plugins/Experimental/NNEDenoiser/Source/NNEDenoiser/Private/NNEDenoiserResourceMapping.h
--- a/Engine/Plugins/Experimental/NNEDenoiser/Source/NNEDenoiser/Private/NNEDenoiserResourceMapping.h
+++ b/Engine/Plugins/Experimental/NNEDenoiser/Source/NNEDenoiser/Private/NNEDenoiserResourceMapping.h
@@ -100,6 +100,16 @@ public:
 		return InputMapping.Add_GetRef(MoveTemp(ResourceMapping));
 	}
 
+	const FResourceMapping* Get(int32 InputIndex) const
+	{
+		if (InputIndex >= 0 && InputIndex < InputMapping.Num())
+		{
+			return &InputMapping[InputIndex];
+		}
+
+		return {};
+	}
+
 	const FResourceMapping& GetChecked(int32 InputIndex) const
 	{
 		check(InputIndex >= 0 && InputIndex < InputMapping.Num());
diff --git a/Engine/Plugins/Experimental/NNEDenoiser/Source/NNEDenoiser/Private/NNEDenoiserViewExtension.cpp b/Engine/Plugins/Experimental/NNEDenoiser/Source/NNEDenoiser/Private/NNEDenoiserViewExtension.cpp
--- a/Engine/Plugins/Experimental/NNEDenoiser/Source/NNEDenoiser/Private/NNEDenoiserViewExtension.cpp
+++ b/Engine/Plugins/Experimental/NNEDenoiser/Source/NNEDenoiser/Private/NNEDenoiserViewExtension.cpp
@@ -32,11 +32,17 @@ static TAutoConsoleVariable<int32> CVarNNEDenoiserModelData(
 	TEXT("Defines the model data of the denoiser model.\n")
 	TEXT("  0: Use model data from Project Settings\n")
 	TEXT("  1: OIDN Color\n")
-	TEXT("  2: OIDN Color and Albedo\n")
-	TEXT("  3: OIDN Color, Albedo and Normal\n")
+	TEXT("  2: OIDN Color and features\n")
+	TEXT("  3: OIDN High Quality (prefilters features)\n")
 	TEXT("  4: OIDN Color | Float16\n")
-	TEXT("  5: OIDN Color and Albedo | Float16\n")
-	TEXT("  6: OIDN Color, Albedo and Normal | Float16")
+	TEXT("  5: OIDN Color and features | Float16\n")
+	TEXT("  6: OIDN High Quality (prefilters features) | Float16\n")
+	TEXT("  7: OIDN Color | Alpha\n")
+	TEXT("  8: OIDN Color and features | Alpha\n")
+	TEXT("  9: OIDN High Quality (prefilters features) | Alpha\n")
+	TEXT("  10: OIDN Color | Alpha | Float16\n")
+	TEXT("  11: OIDN Color and features | Alpha | Float16\n")
+	TEXT("  12: OIDN High Quality (prefilters features) | Alpha | Float16")
 );
 
 static TAutoConsoleVariable<int32> CVarNNEDenoiserRuntimeType(
@@ -68,12 +74,26 @@ FString GetDenoiserModelDataNameFromCVarAndSettings(const UNNEDenoiserSettings*
 	switch(Idx)
 	{
 		case 0: return !Settings->DenoiserModelData.IsNull() ? Settings->DenoiserModelData.ToString() : FString();
+
+		// Float32
 		case 1: return TEXT("/NNEDenoiser/NNEDNN_Oidn2-2-0_Color_Dyn.NNEDNN_Oidn2-2-0_Color_Dyn");
-		case 2: return TEXT("/NNEDenoiser/NNEDNN_Oidn2-2-0_ColorAlbedo_Dyn.NNEDNN_Oidn2-2-0_ColorAlbedo_Dyn");
-		case 3: return TEXT("/NNEDenoiser/NNEDNN_Oidn2-2-0_ColorAlbedoNormal_Dyn.NNEDNN_Oidn2-2-0_ColorAlbedoNormal_Dyn");
+		case 2: return TEXT("/NNEDenoiser/NNEDNN_Oidn2-2-0_ColorAlbedoNormal_Dyn.NNEDNN_Oidn2-2-0_ColorAlbedoNormal_Dyn");
+		case 3: return TEXT("/NNEDenoiser/NNEDNN_Oidn2-2-0_HighQuality_Dyn.NNEDNN_Oidn2-2-0_HighQuality_Dyn");
+		
+		// Float16
 		case 4: return TEXT("/NNEDenoiser/NNEDNN_Oidn2-2-0_Color_Dyn_Float16.NNEDNN_Oidn2-2-0_Color_Dyn_Float16");
-		case 5: return TEXT("/NNEDenoiser/NNEDNN_Oidn2-2-0_ColorAlbedo_Dyn_Float16.NNEDNN_Oidn2-2-0_ColorAlbedo_Dyn_Float16");
-		case 6: return TEXT("/NNEDenoiser/NNEDNN_Oidn2-2-0_ColorAlbedoNormal_Dyn_Float16.NNEDNN_Oidn2-2-0_ColorAlbedoNormal_Dyn_Float16");
+		case 5: return TEXT("/NNEDenoiser/NNEDNN_Oidn2-2-0_ColorAlbedoNormal_Dyn_Float16.NNEDNN_Oidn2-2-0_ColorAlbedoNormal_Dyn_Float16");
+		case 6: return TEXT("/NNEDenoiser/NNEDNN_Oidn2-2-0_HighQuality_Dyn_Float16.NNEDNN_Oidn2-2-0_HighQuality_Dyn_Float16");
+
+		// Alpha, Float32
+		case 7: return TEXT("/NNEDenoiser/NNEDNN_Oidn2-2-0_Color_Dyn_Alpha.NNEDNN_Oidn2-2-0_Color_Dyn_Alpha");
+		case 8: return TEXT("/NNEDenoiser/NNEDNN_Oidn2-2-0_ColorAlbedoNormal_Dyn_Alpha.NNEDNN_Oidn2-2-0_ColorAlbedoNormal_Dyn_Alpha");
+		case 9: return TEXT("/NNEDenoiser/NNEDNN_Oidn2-2-0_HighQuality_Dyn_Alpha.NNEDNN_Oidn2-2-0_HighQuality_Dyn_Alpha");
+		
+		// Alpha, Float16
+		case 10: return TEXT("/NNEDenoiser/NNEDNN_Oidn2-2-0_Color_Dyn_Alpha_Float16.NNEDNN_Oidn2-2-0_Color_Dyn_Alpha_Float16");
+		case 11: return TEXT("/NNEDenoiser/NNEDNN_Oidn2-2-0_ColorAlbedoNormal_Dyn_Alpha_Float16.NNEDNN_Oidn2-2-0_ColorAlbedoNormal_Dyn_Alpha_Float16");
+		case 12: return TEXT("/NNEDenoiser/NNEDNN_Oidn2-2-0_HighQuality_Dyn_Alpha_Float16.NNEDNN_Oidn2-2-0_HighQuality_Dyn_Alpha_Float16");
 	}
 	check(false);
 	return FString();
diff --git a/Engine/Plugins/Experimental/NNEDenoiser/Source/NNEDenoiser/Public/NNEDenoiserSettings.h b/Engine/Plugins/Experimental/NNEDenoiser/Source/NNEDenoiser/Public/NNEDenoiserSettings.h
--- a/Engine/Plugins/Experimental/NNEDenoiser/Source/NNEDenoiser/Public/NNEDenoiserSettings.h
+++ b/Engine/Plugins/Experimental/NNEDenoiser/Source/NNEDenoiser/Public/NNEDenoiserSettings.h
@@ -16,8 +16,8 @@ enum EDenoiserRuntimeType : uint8
 	RDG
 };
 
-/** Settings used to create a NNE Denoise */
-UCLASS(Config = Engine, meta = (DisplayName = "NNE Denoise"))
+/** Settings used to create a NNE Denoiser */
+UCLASS(Config = Engine, meta = (DisplayName = "NNE Denoiser"))
 class NNEDENOISER_API UNNEDenoiserSettings : public UDeveloperSettingsBackedByCVars
 {
 	GENERATED_BODY()
@@ -27,16 +27,16 @@ public:
 
 	virtual void PostInitProperties() override;
 
-	/** Denoiser model data used to create a NNE Denoise */
-	UPROPERTY(Config, EditAnywhere, Category = "NNE Denoise", meta = (DisplayName = "Denoiser model data", ToolTip = "Select the denoiser model data"))
+	/** Denoiser model data used to create a NNE Denoiser */
+	UPROPERTY(Config, EditAnywhere, Category = "NNE Denoiser", meta = (DisplayName = "Denoiser model data", ToolTip = "Select the denoiser model data"))
 	TSoftObjectPtr<UNNEDenoiserModelData> DenoiserModelData;
 
 private:
-	/** Runtime type used to run the NNE Denoise model. Backed by the console variable 'NNEDenoiser.Runtime.Type'. */
-	UPROPERTY(Config, EditAnywhere, Category = "NNE Denoise", meta = (DisplayName = "Runtime Type", ToolTip = "Select a Runtime type", ConsoleVariable = "NNEDenoiser.Runtime.Type"))
+	/** Runtime type used to run the NNE Denoiser model. Backed by the console variable 'NNEDenoiser.Runtime.Type'. */
+	UPROPERTY(Config, EditAnywhere, Category = "NNE Denoiser", meta = (DisplayName = "Runtime Type", ToolTip = "Select a Runtime type", ConsoleVariable = "NNEDenoiser.Runtime.Type"))
 	TEnumAsByte<EDenoiserRuntimeType> RuntimeType;
 
-	/** Runtime name used to run the NNE Denoise model. Backed by the console variable 'NNEDenoiser.Runtime.Name'. */
-	UPROPERTY(Config, EditAnywhere, Category = "NNE Denoise", meta = (DisplayName = "Runtime Name Override", ToolTip = "(Optional) Specify the Runtime name", ConsoleVariable = "NNEDenoiser.Runtime.Name"))
+	/** Runtime name used to run the NNE Denoiser model. Backed by the console variable 'NNEDenoiser.Runtime.Name'. */
+	UPROPERTY(Config, EditAnywhere, Category = "NNE Denoiser", meta = (DisplayName = "Runtime Name Override", ToolTip = "(Optional) Specify the Runtime name", ConsoleVariable = "NNEDenoiser.Runtime.Name"))
 	FString RuntimeName;
 };
diff --git a/Engine/Plugins/Experimental/NNEDenoiser/Source/NNEDenoiserShaders/Internal/NNEDenoiserShadersDefaultCS.h b/Engine/Plugins/Experimental/NNEDenoiser/Source/NNEDenoiserShaders/Internal/NNEDenoiserShadersDefaultCS.h
--- a/Engine/Plugins/Experimental/NNEDenoiser/Source/NNEDenoiserShaders/Internal/NNEDenoiserShadersDefaultCS.h
+++ b/Engine/Plugins/Experimental/NNEDenoiser/Source/NNEDenoiserShaders/Internal/NNEDenoiserShadersDefaultCS.h
@@ -26,49 +26,89 @@ namespace UE::NNEDenoiserShaders::Internal
 		static constexpr int32 MAX_NUM_MAPPED_CHANNELS{ 4 };
 	};
 	
-	class NNEDENOISERSHADERS_API FNNEDenoiserReadInputCS : public FGlobalShader
+	class NNEDENOISERSHADERS_API FNNEDenoiserTextureBufferMappedCopyCS : public FGlobalShader
 	{
-		DECLARE_GLOBAL_SHADER(FNNEDenoiserReadInputCS);
-		SHADER_USE_PARAMETER_STRUCT(FNNEDenoiserReadInputCS, FGlobalShader)
+		DECLARE_GLOBAL_SHADER(FNNEDenoiserTextureBufferMappedCopyCS);
+		SHADER_USE_PARAMETER_STRUCT(FNNEDenoiserTextureBufferMappedCopyCS, FGlobalShader)
 
 		class FNNEDenoiserDataType : SHADER_PERMUTATION_ENUM_CLASS("BUFFER_TYPE_INDEX", ENNEDenoiserDataType);
-		class FNNEDenoiserNumMappedChannels : SHADER_PERMUTATION_RANGE_INT("NUM_MAPPED_CHANNELS", 0, FNNEDenoiserConstants::MAX_NUM_MAPPED_CHANNELS);
+		class FNNEDenoiserNumMappedChannels : SHADER_PERMUTATION_RANGE_INT("NUM_MAPPED_CHANNELS", 0, FNNEDenoiserConstants::MAX_NUM_MAPPED_CHANNELS + 1);
 		using FPermutationDomain = TShaderPermutationDomain<FNNEDenoiserDataType, FNNEDenoiserNumMappedChannels>;
 
 	public:
 
 		BEGIN_SHADER_PARAMETER_STRUCT(FParameters, )
-			SHADER_PARAMETER(int32, InputTextureWidth)
-			SHADER_PARAMETER(int32, InputTextureHeight)
+			SHADER_PARAMETER(int32, Width)
+			SHADER_PARAMETER(int32, Height)
 			SHADER_PARAMETER_RDG_TEXTURE(Texture2D, InputTexture)
-			SHADER_PARAMETER(int32, InputBufferWidth)
-			SHADER_PARAMETER(int32, InputBufferHeight)
-			SHADER_PARAMETER_RDG_BUFFER_UAV(RWBuffer, InputBuffer)
-			SHADER_PARAMETER_ARRAY(FIntVector4, BufferChannel_TextureChannel_Unused_Unused, [FNNEDenoiserConstants::MAX_NUM_MAPPED_CHANNELS])
+			SHADER_PARAMETER_RDG_BUFFER_UAV(RWBuffer, OutputBuffer)
+			SHADER_PARAMETER_ARRAY(FIntVector4, OutputChannel_InputChannel_Unused_Unused, [FNNEDenoiserConstants::MAX_NUM_MAPPED_CHANNELS])
 		END_SHADER_PARAMETER_STRUCT()
 
 		static void ModifyCompilationEnvironment(const FGlobalShaderPermutationParameters& InParameters, FShaderCompilerEnvironment& OutEnvironment);
 	};
 
-	class NNEDENOISERSHADERS_API FNNEDenoiserWriteOutputCS : public FGlobalShader
+	class NNEDENOISERSHADERS_API FNNEDenoiserBufferMappedCopyCS : public FGlobalShader
 	{
-		DECLARE_GLOBAL_SHADER(FNNEDenoiserWriteOutputCS);
-		SHADER_USE_PARAMETER_STRUCT(FNNEDenoiserWriteOutputCS, FGlobalShader)
+		DECLARE_GLOBAL_SHADER(FNNEDenoiserBufferMappedCopyCS);
+		SHADER_USE_PARAMETER_STRUCT(FNNEDenoiserBufferMappedCopyCS, FGlobalShader)
 
 		class FNNEDenoiserDataType : SHADER_PERMUTATION_ENUM_CLASS("BUFFER_TYPE_INDEX", ENNEDenoiserDataType);
-		class FNNEDenoiserNumMappedChannels : SHADER_PERMUTATION_RANGE_INT("NUM_MAPPED_CHANNELS", 0, FNNEDenoiserConstants::MAX_NUM_MAPPED_CHANNELS);
+		class FNNEDenoiserNumMappedChannels : SHADER_PERMUTATION_RANGE_INT("NUM_MAPPED_CHANNELS", 0, FNNEDenoiserConstants::MAX_NUM_MAPPED_CHANNELS + 1);
 		using FPermutationDomain = TShaderPermutationDomain<FNNEDenoiserDataType, FNNEDenoiserNumMappedChannels>;
 
 	public:
 
 		BEGIN_SHADER_PARAMETER_STRUCT(FParameters, )
-			SHADER_PARAMETER(int32, OutputBufferWidth)
-			SHADER_PARAMETER(int32, OutputBufferHeight)
+			SHADER_PARAMETER(int32, Width)
+			SHADER_PARAMETER(int32, Height)
+			SHADER_PARAMETER_RDG_BUFFER_UAV(RWBuffer, InputBuffer)
 			SHADER_PARAMETER_RDG_BUFFER_UAV(RWBuffer, OutputBuffer)
-			SHADER_PARAMETER(int32, OutputTextureWidth)
-			SHADER_PARAMETER(int32, OutputTextureHeight)
+			SHADER_PARAMETER_ARRAY(FIntVector4, OutputChannel_InputChannel_Unused_Unused, [FNNEDenoiserConstants::MAX_NUM_MAPPED_CHANNELS])
+		END_SHADER_PARAMETER_STRUCT()
+
+		static void ModifyCompilationEnvironment(const FGlobalShaderPermutationParameters& InParameters, FShaderCompilerEnvironment& OutEnvironment);
+	};
+
+	class NNEDENOISERSHADERS_API FNNEDenoiserBufferTextureMappedCopyCS : public FGlobalShader
+	{
+		DECLARE_GLOBAL_SHADER(FNNEDenoiserBufferTextureMappedCopyCS);
+		SHADER_USE_PARAMETER_STRUCT(FNNEDenoiserBufferTextureMappedCopyCS, FGlobalShader)
+
+		class FNNEDenoiserDataType : SHADER_PERMUTATION_ENUM_CLASS("BUFFER_TYPE_INDEX", ENNEDenoiserDataType);
+		class FNNEDenoiserNumMappedChannels : SHADER_PERMUTATION_RANGE_INT("NUM_MAPPED_CHANNELS", 0, FNNEDenoiserConstants::MAX_NUM_MAPPED_CHANNELS + 1);
+		using FPermutationDomain = TShaderPermutationDomain<FNNEDenoiserDataType, FNNEDenoiserNumMappedChannels>;
+
+	public:
+
+		BEGIN_SHADER_PARAMETER_STRUCT(FParameters, )
+			SHADER_PARAMETER(int32, Width)
+			SHADER_PARAMETER(int32, Height)
+			SHADER_PARAMETER_RDG_BUFFER_UAV(RWBuffer, InputBuffer)
+			SHADER_PARAMETER_RDG_TEXTURE_UAV(RWTexture2D, OutputTexture)
+			SHADER_PARAMETER_ARRAY(FIntVector4, OutputChannel_InputChannel_Unused_Unused, [FNNEDenoiserConstants::MAX_NUM_MAPPED_CHANNELS])
+		END_SHADER_PARAMETER_STRUCT()
+
+		static void ModifyCompilationEnvironment(const FGlobalShaderPermutationParameters& InParameters, FShaderCompilerEnvironment& OutEnvironment);
+	};
+
+	class NNEDENOISERSHADERS_API FNNEDenoiserTextureMappedCopyCS : public FGlobalShader
+	{
+		DECLARE_GLOBAL_SHADER(FNNEDenoiserTextureMappedCopyCS);
+		SHADER_USE_PARAMETER_STRUCT(FNNEDenoiserTextureMappedCopyCS, FGlobalShader)
+
+		class FNNEDenoiserDataType : SHADER_PERMUTATION_ENUM_CLASS("BUFFER_TYPE_INDEX", ENNEDenoiserDataType);
+		class FNNEDenoiserNumMappedChannels : SHADER_PERMUTATION_RANGE_INT("NUM_MAPPED_CHANNELS", 0, FNNEDenoiserConstants::MAX_NUM_MAPPED_CHANNELS + 1);
+		using FPermutationDomain = TShaderPermutationDomain<FNNEDenoiserDataType, FNNEDenoiserNumMappedChannels>;
+
+	public:
+
+		BEGIN_SHADER_PARAMETER_STRUCT(FParameters, )
+			SHADER_PARAMETER(int32, Width)
+			SHADER_PARAMETER(int32, Height)
+			SHADER_PARAMETER_RDG_TEXTURE_UAV(RWTexture2D, InputTexture)
 			SHADER_PARAMETER_RDG_TEXTURE_UAV(RWTexture2D, OutputTexture)
-			SHADER_PARAMETER_ARRAY(FIntVector4, BufferChannel_TextureChannel_Unused_Unused, [FNNEDenoiserConstants::MAX_NUM_MAPPED_CHANNELS])
+			SHADER_PARAMETER_ARRAY(FIntVector4, OutputChannel_InputChannel_Unused_Unused, [FNNEDenoiserConstants::MAX_NUM_MAPPED_CHANNELS])
 		END_SHADER_PARAMETER_STRUCT()
 
 		static void ModifyCompilationEnvironment(const FGlobalShaderPermutationParameters& InParameters, FShaderCompilerEnvironment& OutEnvironment);
diff --git a/Engine/Plugins/Experimental/NNEDenoiser/Source/NNEDenoiserShaders/Private/NNEDenoiserShadersDefaultCS.cpp b/Engine/Plugins/Experimental/NNEDenoiser/Source/NNEDenoiserShaders/Private/NNEDenoiserShadersDefaultCS.cpp
--- a/Engine/Plugins/Experimental/NNEDenoiser/Source/NNEDenoiserShaders/Private/NNEDenoiserShadersDefaultCS.cpp
+++ b/Engine/Plugins/Experimental/NNEDenoiser/Source/NNEDenoiserShaders/Private/NNEDenoiserShadersDefaultCS.cpp
@@ -4,21 +4,45 @@
 
 namespace UE::NNEDenoiserShaders::Internal
 {
-	void FNNEDenoiserReadInputCS::ModifyCompilationEnvironment(const FGlobalShaderPermutationParameters& InParameters, FShaderCompilerEnvironment& OutEnvironment)
+	void CommonModifyCompilationEnvironment(const FGlobalShaderPermutationParameters& InParameters, FShaderCompilerEnvironment& OutEnvironment)
 	{
 		FGlobalShader::ModifyCompilationEnvironment(InParameters, OutEnvironment);
 		OutEnvironment.SetDefine(TEXT("THREAD_GROUP_SIZE"), FNNEDenoiserConstants::THREAD_GROUP_SIZE);
 		OutEnvironment.SetDefine(TEXT("MAX_NUM_MAPPED_CHANNELS"), FNNEDenoiserConstants::MAX_NUM_MAPPED_CHANNELS);
 	}
 
-	void FNNEDenoiserWriteOutputCS::ModifyCompilationEnvironment(const FGlobalShaderPermutationParameters& InParameters, FShaderCompilerEnvironment& OutEnvironment)
+	void FNNEDenoiserTextureBufferMappedCopyCS::ModifyCompilationEnvironment(const FGlobalShaderPermutationParameters& InParameters, FShaderCompilerEnvironment& OutEnvironment)
 	{
-		FGlobalShader::ModifyCompilationEnvironment(InParameters, OutEnvironment);
-		OutEnvironment.SetDefine(TEXT("THREAD_GROUP_SIZE"), FNNEDenoiserConstants::THREAD_GROUP_SIZE);
-		OutEnvironment.SetDefine(TEXT("MAX_NUM_MAPPED_CHANNELS"), FNNEDenoiserConstants::MAX_NUM_MAPPED_CHANNELS);
+		CommonModifyCompilationEnvironment(InParameters, OutEnvironment);
+		OutEnvironment.SetDefine(TEXT("INPUT_TYPE"), 0);
+		OutEnvironment.SetDefine(TEXT("OUTPUT_TYPE"), 1);
 	}
 
-	IMPLEMENT_GLOBAL_SHADER(FNNEDenoiserReadInputCS, "/NNEDenoiserShaders/NNEDenoiserShadersDefault.usf", "ReadInput", SF_Compute);
-	IMPLEMENT_GLOBAL_SHADER(FNNEDenoiserWriteOutputCS, "/NNEDenoiserShaders/NNEDenoiserShadersDefault.usf", "WriteOutput", SF_Compute);
+	void FNNEDenoiserBufferMappedCopyCS::ModifyCompilationEnvironment(const FGlobalShaderPermutationParameters& InParameters, FShaderCompilerEnvironment& OutEnvironment)
+	{
+		CommonModifyCompilationEnvironment(InParameters, OutEnvironment);
+		OutEnvironment.SetDefine(TEXT("INPUT_TYPE"), 1);
+		OutEnvironment.SetDefine(TEXT("OUTPUT_TYPE"), 1);
+	}
+
+	void FNNEDenoiserBufferTextureMappedCopyCS::ModifyCompilationEnvironment(const FGlobalShaderPermutationParameters& InParameters, FShaderCompilerEnvironment& OutEnvironment)
+	{
+		CommonModifyCompilationEnvironment(InParameters, OutEnvironment);
+		OutEnvironment.SetDefine(TEXT("INPUT_TYPE"), 1);
+		OutEnvironment.SetDefine(TEXT("OUTPUT_TYPE"), 0);
+	}
+
+	void FNNEDenoiserTextureMappedCopyCS::ModifyCompilationEnvironment(const FGlobalShaderPermutationParameters& InParameters, FShaderCompilerEnvironment& OutEnvironment)
+	{
+		CommonModifyCompilationEnvironment(InParameters, OutEnvironment);
+		OutEnvironment.SetDefine(TEXT("INPUT_TYPE"), 0);
+		OutEnvironment.SetDefine(TEXT("OUTPUT_TYPE"), 0);
+	}
+
+	IMPLEMENT_GLOBAL_SHADER(FNNEDenoiserTextureBufferMappedCopyCS, "/NNEDenoiserShaders/NNEDenoiserShadersDefault.usf", "MappedCopy", SF_Compute);
+	IMPLEMENT_GLOBAL_SHADER(FNNEDenoiserBufferMappedCopyCS, "/NNEDenoiserShaders/NNEDenoiserShadersDefault.usf", "MappedCopy", SF_Compute);
+
+	IMPLEMENT_GLOBAL_SHADER(FNNEDenoiserBufferTextureMappedCopyCS, "/NNEDenoiserShaders/NNEDenoiserShadersDefault.usf", "MappedCopy", SF_Compute);
+	IMPLEMENT_GLOBAL_SHADER(FNNEDenoiserTextureMappedCopyCS, "/NNEDenoiserShaders/NNEDenoiserShadersDefault.usf", "MappedCopy", SF_Compute);
 
 } // UE::NNEDenoiser::Private
\ No newline at end of file
diff --git a/Engine/Plugins/Experimental/NNEDenoiser/Source/NNEDenoiserShaders/Shaders/NNEDenoiserShadersDefault.usf b/Engine/Plugins/Experimental/NNEDenoiser/Source/NNEDenoiserShaders/Shaders/NNEDenoiserShadersDefault.usf
--- a/Engine/Plugins/Experimental/NNEDenoiser/Source/NNEDenoiserShaders/Shaders/NNEDenoiserShadersDefault.usf
+++ b/Engine/Plugins/Experimental/NNEDenoiser/Source/NNEDenoiserShaders/Shaders/NNEDenoiserShadersDefault.usf
@@ -3,78 +3,67 @@
 #include "/Engine/Public/Platform.ush"
 #include "/Engine/Private/Common.ush"
 
-#if BUFFER_TYPE_INDEX == 4
+#if BUFFER_TYPE_INDEX == 4 // Float16
 	#define BUFFER_TYPE half
 	#define READ(X) float(X)
 	#define WRITE(X) half(X)
 #else
-	// BUFFER_TYPE_INDEX == 3
+	// BUFFER_TYPE_INDEX == 3 // Float32
 	#define BUFFER_TYPE float
 	#define READ(X) X
 	#define WRITE(X) X
 #endif
 
-int InputTextureWidth;
-int InputTextureHeight;
-Texture2D InputTexture;
+int Width;
+int Height;
+
+int4 OutputChannel_InputChannel_Unused_Unused[MAX_NUM_MAPPED_CHANNELS];
 
-int InputBufferWidth;
-int InputBufferHeight;
+#if INPUT_TYPE == 0 // Texture
+Texture2D InputTexture;
+#else // INPUT_TYPE == 1 // Buffer
 RWBuffer<BUFFER_TYPE> InputBuffer;
+#endif
 
-int4 BufferChannel_TextureChannel_Unused_Unused[MAX_NUM_MAPPED_CHANNELS];
+#if OUTPUT_TYPE == 0 // Texture
+RWTexture2D<float4> OutputTexture;
+#else // OUTPUT_TYPE == 1 // Buffer
+RWBuffer<BUFFER_TYPE> OutputBuffer;
+#endif
 
 [numthreads(THREAD_GROUP_SIZE, THREAD_GROUP_SIZE, 1)]
-void ReadInput(in const uint3 DispatchThreadID : SV_DispatchThreadID)
+void MappedCopy(in const uint3 DispatchThreadID : SV_DispatchThreadID)
 {
-	if (DispatchThreadID.x >= InputBufferWidth || DispatchThreadID.y >= InputBufferHeight)
+	if (DispatchThreadID.x >= Width || DispatchThreadID.y >= Height)
 	{
 		return;
 	}
 
-	const int Idx = DispatchThreadID.y * InputBufferWidth + DispatchThreadID.x;
-	const int Offset = InputBufferWidth * InputBufferHeight;
+#if INPUT_TYPE == 1 || OUTPUT_TYPE == 1
+	const int Idx = DispatchThreadID.y * Width + DispatchThreadID.x;
+	const int Offset = Width * Height;
+#endif
 
-	const float4 BufferColor = InputTexture.Load(int3(DispatchThreadID.x, DispatchThreadID.y, 0));
+#if INPUT_TYPE == 0
+	const float4 InputColor = InputTexture.Load(int3(DispatchThreadID.x, DispatchThreadID.y, 0));
+#endif
 
 	for (int I = 0; I < NUM_MAPPED_CHANNELS; I++)
 	{
-		const int4 Info = BufferChannel_TextureChannel_Unused_Unused[I];
-		const int BufferChannel = Info.x;
-		const int TextureChannel = Info.y;
-
-		InputBuffer[Idx + BufferChannel * Offset] = WRITE(BufferColor[TextureChannel]);
-	}
-}
-
-int OutputBufferWidth;
-int OutputBufferHeight;
-RWBuffer<BUFFER_TYPE> OutputBuffer;
-
-int OutputTextureWidth;
-int OutputTextureHeight;
-RWTexture2D<float4> OutputTexture;
+		const int4 Info = OutputChannel_InputChannel_Unused_Unused[I];
+		const int OutputChannel = Info.x;
+		const int InputChannel = Info.y;
 
-[numthreads(THREAD_GROUP_SIZE, THREAD_GROUP_SIZE, 1)]
-void WriteOutput(in const uint3 DispatchThreadID : SV_DispatchThreadID)
-{
-	if (DispatchThreadID.x >= OutputTextureWidth || DispatchThreadID.y >= OutputTextureHeight)
-	{
-		return;
-	}
-
-	const int Idx = DispatchThreadID.y * OutputBufferWidth + DispatchThreadID.x;
-	const int Offset = OutputBufferWidth * OutputBufferHeight;
-
-	for (int I = 0; I < NUM_MAPPED_CHANNELS; I++)
-	{
-		const int4 Info = BufferChannel_TextureChannel_Unused_Unused[I];
-		const int BufferChannel = Info.x;
-		const int TextureChannel = Info.y;
+#if INPUT_TYPE == 0
+		const float InputValue = InputColor[InputChannel];
+#else
+		const float InputValue = InputBuffer[Idx + InputChannel * Offset];
+#endif
 
-		OutputTexture[DispatchThreadID.xy][TextureChannel] = READ(OutputBuffer[Idx + BufferChannel * Offset]);
+#if OUTPUT_TYPE == 0
+		OutputTexture[DispatchThreadID.xy][OutputChannel] = READ(InputValue);
+#else
+		OutputBuffer[Idx + OutputChannel * Offset] = WRITE(InputValue);
+#endif
 	}
-
-	// note: alpha channel denoising or transfer not supported (yet)
-	OutputTexture[DispatchThreadID.xy].a = 1.0;
 }
\ No newline at end of file
