[NNEDenosier] Fix NNEDenoiser HDR color image input scale for OIDN models. Calculate it from average of log luminance (auto exposure).

#tests Editor, Win64 compared output with OIDN plugin and no perf regression

#rnx

diff --git a/Engine/Plugins/Experimental/NNEDenoiser/Source/NNEDenoiser/Private/NNEDenoiserAutoExposure.cpp b/Engine/Plugins/Experimental/NNEDenoiser/Source/NNEDenoiser/Private/NNEDenoiserAutoExposure.cpp
new file mode 100644
--- /dev/null
+++ b/Engine/Plugins/Experimental/NNEDenoiser/Source/NNEDenoiser/Private/NNEDenoiserAutoExposure.cpp
@@ -0,0 +1,176 @@
+// Copyright Epic Games, Inc. All Rights Reserved.
+
+#include "NNEDenoiserAutoExposure.h"
+#include "NNEDenoiserLog.h"
+#include "NNEDenoiserShadersAutoExposureCS.h"
+#include "NNEDenoiserUtils.h"
+#include "RenderGraphBuilder.h"
+
+DECLARE_GPU_STAT_NAMED(FAutoExposure, TEXT("AutoExposure"));
+DECLARE_GPU_STAT_NAMED(FAutoExposureDownsample, TEXT("AutoExposure.Downsample"));
+DECLARE_GPU_STAT_NAMED(FAutoExposureReduce, TEXT("AutoExposure.Reduce"));
+DECLARE_GPU_STAT_NAMED(FAutoExposureReduceFinal, TEXT("AutoExposure.ReduceFinal"));
+
+BEGIN_SHADER_PARAMETER_STRUCT(FAutoExposureTestDownloadParameters, )
+	RDG_BUFFER_ACCESS(OutputBuffer, ERHIAccess::CopySrc)
+END_SHADER_PARAMETER_STRUCT()
+
+namespace UE::NNEDenoiser::Private
+{
+
+void FAutoExposure::Run(TConstArrayView<FLinearColor> InputData, FIntPoint Size, float& OutputValue) const
+{
+	using namespace UE::NNEDenoiserShaders::Internal;
+
+	OutputValue = 0.0f;
+	if (InputData.IsEmpty())
+	{
+		return;
+	}
+
+	const int32 NumBinsW = FMath::DivideAndRoundUp(Size.X, FAutoExposureDownsampleConstants::MAX_BIN_SIZE);
+	const int32 NumBinsH = FMath::DivideAndRoundUp(Size.Y, FAutoExposureDownsampleConstants::MAX_BIN_SIZE);
+	const int32 NumBins = NumBinsW * NumBinsH;
+
+	TArray<float> Bins;
+	Bins.SetNumUninitialized(NumBins);
+
+	for (int32 GroupX = 0; GroupX < NumBinsW; GroupX++)
+	{
+		for (int32 GroupY = 0; GroupY < NumBinsH; GroupY++)
+		{
+			const int32 BeginW = GroupX * Size.X / NumBinsW;
+			const int32 BeginH = GroupY * Size.Y / NumBinsH;
+			const int32 EndW = (GroupX + 1) * Size.X / NumBinsW;
+			const int32 EndH = (GroupY + 1) * Size.Y / NumBinsH;
+
+			float Lum = 0.0f;
+			for (int32 i = 0; i < FAutoExposureDownsampleConstants::MAX_BIN_SIZE; i++)
+			{
+				for (int32 j = 0; j < FAutoExposureDownsampleConstants::MAX_BIN_SIZE; j++)
+				{
+					const int32 CoordW = BeginW + i;
+					const int32 CoordH = BeginH + j;
+
+					if (CoordW < EndW && CoordH < EndH)
+					{
+						FLinearColor Color = InputData[CoordH * Size.X + CoordW];
+						
+						Lum += 0.212671f * Color.R + 0.715160f * Color.G + 0.072169f * Color.B;
+					}
+				}
+			}
+
+			Bins[GroupY * NumBinsW + GroupX] = Lum / float((EndW - BeginW) * (EndH - BeginH));
+		}
+	}
+
+	float Sum = 0.0;
+	int32 Count = 0;
+	for (int32 i = 0; i < NumBins; i++)
+	{
+		const float Lum = Bins[i];
+		if (Lum > FAutoExposureReduceConstants::EPS)
+		{
+			Sum += FMath::Log2(Lum);
+			Count++;
+		}
+	}
+
+	OutputValue = Count > 0 ? FAutoExposureReduceConstants::KEY / FMath::Exp2(Sum / (float)Count) : 1.0f;
+}
+
+void FAutoExposure::EnqueueRDG(FRDGBuilder& GraphBuilder, FRDGTextureRef InputTexture, FRDGBufferRef OutputBuffer) const
+{
+	using namespace UE::NNEDenoiserShaders::Internal;
+
+	const FIntVector InputTextureSize = InputTexture->Desc.GetSize();
+
+	const int32 NumBinsW = FMath::DivideAndRoundUp(InputTextureSize.X, FAutoExposureDownsampleConstants::MAX_BIN_SIZE);
+	const int32 NumBinsH = FMath::DivideAndRoundUp(InputTextureSize.Y, FAutoExposureDownsampleConstants::MAX_BIN_SIZE);
+	const int32 NumBins = NumBinsW * NumBinsH;
+
+	const int NumReduceGroups = FMath::Min(CeilDiv(NumBins, FAutoExposureReduceConstants::THREAD_GROUP_SIZE), FAutoExposureReduceConstants::THREAD_GROUP_SIZE);
+
+	FRDGBufferDesc OutputBinsDesc = FRDGBufferDesc::CreateBufferDesc(sizeof(float), NumBins);
+	FRDGBufferRef OutputBins = GraphBuilder.CreateBuffer(OutputBinsDesc, TEXT("AutoExposureOutputBins"));
+
+	FRDGBufferDesc OutputSumsDesc = FRDGBufferDesc::CreateBufferDesc(sizeof(float), NumReduceGroups);
+	FRDGBufferRef OutputSums = GraphBuilder.CreateBuffer(OutputSumsDesc, TEXT("AutoExposureOutputSums"));
+
+	FRDGBufferDesc OutputCountsDesc = FRDGBufferDesc::CreateBufferDesc(sizeof(int32), NumReduceGroups);
+	FRDGBufferRef OutputCounts = GraphBuilder.CreateBuffer(OutputCountsDesc, TEXT("AutoExposureOutputCounts"));
+
+	{
+		const FIntVector GroupCount = FIntVector(NumBinsW, NumBinsH, 1);
+
+		FAutoExposureDownsampleCS::FParameters *ShaderParameters = GraphBuilder.AllocParameters<FAutoExposureDownsampleCS::FParameters>();
+		ShaderParameters->InputTextureWidth = InputTextureSize.X;
+		ShaderParameters->InputTextureHeight = InputTextureSize.Y;
+		ShaderParameters->InputTexture = InputTexture;
+		ShaderParameters->NumBinsW = NumBinsW;
+		ShaderParameters->NumBinsH = NumBinsH;
+		ShaderParameters->OutputBins = GraphBuilder.CreateUAV(FRDGBufferUAVDesc(OutputBins, EPixelFormat::PF_R32_FLOAT));
+
+		FGlobalShaderMap* GlobalShaderMap = GetGlobalShaderMap(GMaxRHIFeatureLevel);
+		TShaderMapRef<FAutoExposureDownsampleCS> Shader(GlobalShaderMap);
+
+		RDG_EVENT_SCOPE(GraphBuilder, "AutoExposure.Downsample");
+		RDG_GPU_STAT_SCOPE(GraphBuilder, FAutoExposureDownsample);
+
+		FComputeShaderUtils::AddPass(
+			GraphBuilder,
+			RDG_EVENT_NAME("AutoExposure.Downsample"),
+			ERDGPassFlags::Compute | ERDGPassFlags::NeverCull,
+			Shader,
+			ShaderParameters,
+			GroupCount);
+	}
+
+	{
+		FAutoExposureReduceCS::FParameters *ShaderParameters = GraphBuilder.AllocParameters<FAutoExposureReduceCS::FParameters>();
+		ShaderParameters->InputSize = NumBins;
+		ShaderParameters->InputBins = GraphBuilder.CreateUAV(FRDGBufferUAVDesc(OutputBins, EPixelFormat::PF_R32_FLOAT));
+		ShaderParameters->OutputSums = GraphBuilder.CreateUAV(FRDGBufferUAVDesc(OutputSums, EPixelFormat::PF_R32_FLOAT));
+		ShaderParameters->OutputCounts = GraphBuilder.CreateUAV(FRDGBufferUAVDesc(OutputCounts, EPixelFormat::PF_R32_SINT));
+		ShaderParameters->NumThreads = NumReduceGroups * FAutoExposureDownsampleConstants::MAX_BIN_SIZE;
+
+		FGlobalShaderMap* GlobalShaderMap = GetGlobalShaderMap(GMaxRHIFeatureLevel);
+		TShaderMapRef<FAutoExposureReduceCS> Shader(GlobalShaderMap);
+
+		RDG_EVENT_SCOPE(GraphBuilder, "AutoExposure.Recude");
+		RDG_GPU_STAT_SCOPE(GraphBuilder, FAutoExposureReduce);
+
+		FComputeShaderUtils::AddPass(
+			GraphBuilder,
+			RDG_EVENT_NAME("AutoExposure.Recude"),
+			ERDGPassFlags::Compute | ERDGPassFlags::NeverCull,
+			Shader,
+			ShaderParameters,
+			FIntVector(NumReduceGroups, 1, 1));
+	}
+
+	{
+		FAutoExposureReduceFinalCS::FParameters *ShaderParameters = GraphBuilder.AllocParameters<FAutoExposureReduceFinalCS::FParameters>();
+		ShaderParameters->InputSize = NumReduceGroups;
+		ShaderParameters->InputSums = GraphBuilder.CreateUAV(FRDGBufferUAVDesc(OutputSums, EPixelFormat::PF_R32_FLOAT));
+		ShaderParameters->InputCounts = GraphBuilder.CreateUAV(FRDGBufferUAVDesc(OutputCounts, EPixelFormat::PF_R32_SINT));
+		ShaderParameters->OutputBuffer = GraphBuilder.CreateUAV(FRDGBufferUAVDesc(OutputBuffer, EPixelFormat::PF_R32_FLOAT));
+
+		FGlobalShaderMap* GlobalShaderMap = GetGlobalShaderMap(GMaxRHIFeatureLevel);
+		TShaderMapRef<FAutoExposureReduceFinalCS> Shader(GlobalShaderMap);
+
+		RDG_EVENT_SCOPE(GraphBuilder, "AutoExposure.RecudeFinal");
+		RDG_GPU_STAT_SCOPE(GraphBuilder, FAutoExposureReduceFinal);
+
+		FComputeShaderUtils::AddPass(
+			GraphBuilder,
+			RDG_EVENT_NAME("AutoExposure.RecudeFinal"),
+			ERDGPassFlags::Compute | ERDGPassFlags::NeverCull,
+			Shader,
+			ShaderParameters,
+			FIntVector(1, 1, 1));
+	}
+}
+
+}
\ No newline at end of file
diff --git a/Engine/Plugins/Experimental/NNEDenoiser/Source/NNEDenoiser/Private/NNEDenoiserAutoExposure.h b/Engine/Plugins/Experimental/NNEDenoiser/Source/NNEDenoiser/Private/NNEDenoiserAutoExposure.h
new file mode 100644
--- /dev/null
+++ b/Engine/Plugins/Experimental/NNEDenoiser/Source/NNEDenoiser/Private/NNEDenoiserAutoExposure.h
@@ -0,0 +1,32 @@
+// Copyright Epic Games, Inc. All Rights Reserved.
+
+#pragma once
+
+#include "Containers/ContainersFwd.h"
+#include "Math/MathFwd.h"
+#include "RenderGraphFwd.h"
+
+struct FLinearColor;
+
+namespace UE::NNEDenoiser::Private
+{
+
+class IAutoExposure
+{
+public:
+	virtual ~IAutoExposure() = default;
+
+	virtual void Run(TConstArrayView<FLinearColor> InputData, FIntPoint Size, float& OutputValue) const = 0;
+	virtual void EnqueueRDG(FRDGBuilder& GraphBuilder, FRDGTextureRef InputTexture, FRDGBufferRef OutputBuffer) const = 0;
+};
+
+class FAutoExposure : public IAutoExposure
+{
+public:
+	virtual ~FAutoExposure() = default;
+
+	virtual void Run(TConstArrayView<FLinearColor> InputData, FIntPoint Size, float& OutputValue) const override;
+	virtual void EnqueueRDG(FRDGBuilder& GraphBuilder, FRDGTextureRef InputTexture, FRDGBufferRef OutputBuffer) const override;
+};
+
+}
\ No newline at end of file
diff --git a/Engine/Plugins/Experimental/NNEDenoiser/Source/NNEDenoiser/Private/NNEDenoiserGenericDenoiser.cpp b/Engine/Plugins/Experimental/NNEDenoiser/Source/NNEDenoiser/Private/NNEDenoiserGenericDenoiser.cpp
--- a/Engine/Plugins/Experimental/NNEDenoiser/Source/NNEDenoiser/Private/NNEDenoiserGenericDenoiser.cpp
+++ b/Engine/Plugins/Experimental/NNEDenoiser/Source/NNEDenoiser/Private/NNEDenoiserGenericDenoiser.cpp
@@ -2,6 +2,7 @@
 
 #include "NNEDenoiserGenericDenoiser.h"
 #include "Algo/Transform.h"
+#include "NNEDenoiserAutoExposure.h"
 #include "NNEDenoiserHistory.h"
 #include "NNEDenoiserIOProcess.h"
 #include "NNEDenoiserLog.h"
@@ -131,8 +132,19 @@ void AddTilePasses(FRDGBuilder& GraphBuilder, IModelInstance& ModelInstance, con
 	OutputProcess.AddPasses(GraphBuilder, ModelInstance.GetOutputTensorDescs(), GetOutputTensorShapes(ModelInstance), ResourceAccess, OutputBuffers, OutputTexture);
 }
 
-FGenericDenoiser::FGenericDenoiser(TUniquePtr<IModelInstance> ModelInstance, TUniquePtr<IInputProcess> InputProcess, TUniquePtr<IOutputProcess> OutputProcess, FParameters DenoiserParameters) :
-	ModelInstance(MoveTemp(ModelInstance)), InputProcess(MoveTemp(InputProcess)), OutputProcess(MoveTemp(OutputProcess)), DenoiserParameters(DenoiserParameters)
+FGenericDenoiser::FGenericDenoiser(
+	TUniquePtr<IModelInstance> ModelInstance,
+	TUniquePtr<IInputProcess> InputProcess,
+	TUniquePtr<IOutputProcess> OutputProcess,
+	FParameters DenoiserParameters,
+	TUniquePtr<IAutoExposure> AutoExposure,
+	TSharedPtr<ITransferFunction> TransferFunction) :
+		ModelInstance(MoveTemp(ModelInstance)),
+		InputProcess(MoveTemp(InputProcess)),
+		OutputProcess(MoveTemp(OutputProcess)),
+		DenoiserParameters(DenoiserParameters),
+		AutoExposure(MoveTemp(AutoExposure)),
+		TransferFunction(TransferFunction)
 {
 
 }
@@ -229,6 +241,20 @@ TUniquePtr<FHistory> FGenericDenoiser::AddPasses(
 	TArray<FRDGBufferRef> InputBuffers = CreateBuffersRDG(GraphBuilder, ModelInstance->GetInputTensorDescs(), ModelInstance->GetInputTensorShapes());
 	TArray<FRDGBufferRef> OutputBuffers = CreateBuffersRDG(GraphBuilder, ModelInstance->GetOutputTensorDescs(), GetOutputTensorShapes(*ModelInstance));
 
+	if (TransferFunction.IsValid() && AutoExposure.IsValid())
+	{
+		FRDGBufferDesc InputBufferDesc = FRDGBufferDesc::CreateBufferDesc(sizeof(float), 2);
+		FRDGBufferRef InputScaleBuffer = GraphBuilder.CreateBuffer(InputBufferDesc, TEXT("AutoExposureOutputBuffer"));
+
+		AutoExposure->EnqueueRDG(GraphBuilder, ColorTex, InputScaleBuffer);
+
+		TransferFunction->RDGSetInputScale(InputScaleBuffer);
+	}
+	else
+	{
+		checkf(!TransferFunction.IsValid() && !AutoExposure.IsValid(), TEXT("TransferFunction and AutoExposure either both need to be set or not set."))
+	}
+
 	for (int32 I = 0; I < Tiling.Tiles.Num(); I++)
 	{
 		ResourceManager.BeginTile(I);
diff --git a/Engine/Plugins/Experimental/NNEDenoiser/Source/NNEDenoiser/Private/NNEDenoiserGenericDenoiser.h b/Engine/Plugins/Experimental/NNEDenoiser/Source/NNEDenoiser/Private/NNEDenoiserGenericDenoiser.h
--- a/Engine/Plugins/Experimental/NNEDenoiser/Source/NNEDenoiser/Private/NNEDenoiserGenericDenoiser.h
+++ b/Engine/Plugins/Experimental/NNEDenoiser/Source/NNEDenoiser/Private/NNEDenoiserGenericDenoiser.h
@@ -18,11 +18,20 @@ class FHistory;
 class IModelInstance;
 class IInputProcess;
 class IOutputProcess;
+class IAutoExposure;
+class ITransferFunction;
 
 class FGenericDenoiser
 {
 public:
-	FGenericDenoiser(TUniquePtr<IModelInstance> ModelInstance, TUniquePtr<IInputProcess> InputProcess, TUniquePtr<IOutputProcess> OutputProcess, FParameters DenoiserParameters);
+	FGenericDenoiser(
+		TUniquePtr<IModelInstance> ModelInstance,
+		TUniquePtr<IInputProcess> InputProcess,
+		TUniquePtr<IOutputProcess> OutputProcess,
+		FParameters DenoiserParameters,
+		TUniquePtr<IAutoExposure> AutoExposure,
+		TSharedPtr<ITransferFunction> TransferFunction
+	);
 
 	~FGenericDenoiser();
 
@@ -48,6 +57,8 @@ private:
 	TUniquePtr<IInputProcess> InputProcess;
 	TUniquePtr<IOutputProcess> OutputProcess;
 	FParameters DenoiserParameters;
+	TUniquePtr<IAutoExposure> AutoExposure;
+	TSharedPtr<ITransferFunction> TransferFunction;
 
 	FIntPoint LastExtent = { -1, -1 };
 	FTiling Tiling{};
diff --git a/Engine/Plugins/Experimental/NNEDenoiser/Source/NNEDenoiser/Private/NNEDenoiserIOProcessBase.cpp b/Engine/Plugins/Experimental/NNEDenoiser/Source/NNEDenoiser/Private/NNEDenoiserIOProcessBase.cpp
--- a/Engine/Plugins/Experimental/NNEDenoiser/Source/NNEDenoiser/Private/NNEDenoiserIOProcessBase.cpp
+++ b/Engine/Plugins/Experimental/NNEDenoiser/Source/NNEDenoiser/Private/NNEDenoiserIOProcessBase.cpp
@@ -4,12 +4,15 @@
 #include "NNEDenoiserLog.h"
 #include "NNEDenoiserModelInstance.h"
 #include "NNEDenoiserParameters.h"
-#include "NNEDenoiserShadersDefaultCS.h"
+#include "NNEDenoiserShadersMappedCopyCS.h"
+#include "NNEDenoiserShadersDefaultIOProcessCS.h"
+#include "NNEDenoiserTransferFunction.h"
 #include "NNEDenoiserUtils.h"
 #include "RHIStaticStates.h"
 
 DECLARE_GPU_STAT_NAMED(FNNEDenoiserReadInput, TEXT("NNEDenoiser.ReadInput"));
 DECLARE_GPU_STAT_NAMED(FNNEDenoiserWriteOutput, TEXT("NNEDenoiser.WriteOutput"));
+DECLARE_GPU_STAT_NAMED(FNNEDenoiserDefaultIOProcess, TEXT("NNEDenoiser.DefaultIOProcess"));
 
 namespace UE::NNEDenoiser::Private
 {
@@ -17,6 +20,68 @@ namespace UE::NNEDenoiser::Private
 namespace IOProcessBaseHelper
 {
 
+NNEDenoiserShaders::Internal::EDefaultIOProcessInputKind GetInputKind(EResourceName TensorName)
+{
+	using NNEDenoiserShaders::Internal::EDefaultIOProcessInputKind;
+	
+	switch(TensorName)
+	{
+		case EResourceName::Color:	return EDefaultIOProcessInputKind::Color;
+		case EResourceName::Albedo:	return EDefaultIOProcessInputKind::Albedo;
+		case EResourceName::Normal:	return EDefaultIOProcessInputKind::Normal;
+		case EResourceName::Flow:	return EDefaultIOProcessInputKind::Flow;
+		case EResourceName::Output:	return EDefaultIOProcessInputKind::Output;
+	}
+
+	// There should be a case for every resource name
+	checkNoEntry();
+	
+	return EDefaultIOProcessInputKind::Color;
+}
+
+void AddPreOrPostProcess(
+	FRDGBuilder& GraphBuilder,
+	FRDGTextureRef InputTexture,
+	EResourceName TensorName,
+	int32 FrameIdx,
+	FRDGTextureRef OutputTexture)
+{
+	using namespace UE::NNEDenoiserShaders::Internal;
+
+	const FIntVector InputTextureSize = InputTexture->Desc.GetSize();
+	const FIntVector OutputTextureSize = OutputTexture->Desc.GetSize();
+
+	FDefaultIOProcessCS::FParameters *ShaderParameters = GraphBuilder.AllocParameters<FDefaultIOProcessCS::FParameters>();
+	ShaderParameters->InputTextureWidth = InputTextureSize.X;
+	ShaderParameters->InputTextureHeight = InputTextureSize.Y;
+	ShaderParameters->InputTexture = InputTexture;
+	ShaderParameters->OutputTextureWidth = OutputTextureSize.X;
+	ShaderParameters->OutputTextureHeight = OutputTextureSize.Y;
+	ShaderParameters->OutputTexture = GraphBuilder.CreateUAV(FRDGTextureUAVDesc(OutputTexture));
+
+	FDefaultIOProcessCS::FPermutationDomain PermutationVector;
+	PermutationVector.Set<FDefaultIOProcessCS::FDefaultIOProcessInputKind>(GetInputKind(TensorName));
+
+	FIntVector ThreadGroupCount = FIntVector(
+		FMath::DivideAndRoundUp(OutputTextureSize.X, FDefaultIOProcessConstants::THREAD_GROUP_SIZE),
+		FMath::DivideAndRoundUp(OutputTextureSize.Y, FDefaultIOProcessConstants::THREAD_GROUP_SIZE),
+		1);
+
+	FGlobalShaderMap* GlobalShaderMap = GetGlobalShaderMap(GMaxRHIFeatureLevel);
+	TShaderMapRef<FDefaultIOProcessCS> Shader(GlobalShaderMap, PermutationVector);
+
+	RDG_EVENT_SCOPE(GraphBuilder, "NNEDenoiser.DefaultIOProcess");
+	RDG_GPU_STAT_SCOPE(GraphBuilder, FNNEDenoiserDefaultIOProcess);
+
+	FComputeShaderUtils::AddPass(
+		GraphBuilder,
+		RDG_EVENT_NAME("NNEDenoiser.DefaultIOProcess"),
+		ERDGPassFlags::Compute | ERDGPassFlags::NeverCull,
+		Shader,
+		ShaderParameters,
+		ThreadGroupCount);
+}
+
 void AddReadInputPass(
 	FRDGBuilder& GraphBuilder,
 	FRDGTextureRef InputTexture,
@@ -33,7 +98,7 @@ void AddReadInputPass(
 		return;
 	}
 	
-	if (ChannelMapping.Num() >= FNNEDenoiserConstants::MAX_NUM_MAPPED_CHANNELS)
+	if (ChannelMapping.Num() >= FMappedCopyConstants::MAX_NUM_MAPPED_CHANNELS)
 	{
 		UE_LOG(LogNNEDenoiser, Warning, TEXT("AddReadInputPass: ChannelMapping has too many entries!"));
 		return;
@@ -56,8 +121,8 @@ void AddReadInputPass(
 	PermutationVector.Set<FNNEDenoiserTextureBufferMappedCopyCS::FNNEDenoiserNumMappedChannels>(ChannelMapping.Num());
 
 	FIntVector ReadInputThreadGroupCount = FIntVector(
-		FMath::DivideAndRoundUp(InputBufferSize.X, FNNEDenoiserConstants::THREAD_GROUP_SIZE),
-		FMath::DivideAndRoundUp(InputBufferSize.Y, FNNEDenoiserConstants::THREAD_GROUP_SIZE),
+		FMath::DivideAndRoundUp(InputBufferSize.X, FMappedCopyConstants::THREAD_GROUP_SIZE),
+		FMath::DivideAndRoundUp(InputBufferSize.Y, FMappedCopyConstants::THREAD_GROUP_SIZE),
 		1);
 
 	FGlobalShaderMap* GlobalShaderMap = GetGlobalShaderMap(GMaxRHIFeatureLevel);
@@ -102,8 +167,8 @@ void AddWriteOutputPass(
 	PermutationVector.Set<FNNEDenoiserBufferTextureMappedCopyCS::FNNEDenoiserNumMappedChannels>(ChannelMapping.Num());
 
 	FIntVector WriteOutputThreadGroupCount = FIntVector(
-		FMath::DivideAndRoundUp(OutputTextureSize.X, FNNEDenoiserConstants::THREAD_GROUP_SIZE),
-		FMath::DivideAndRoundUp(OutputTextureSize.Y, FNNEDenoiserConstants::THREAD_GROUP_SIZE),
+		FMath::DivideAndRoundUp(OutputTextureSize.X, FMappedCopyConstants::THREAD_GROUP_SIZE),
+		FMath::DivideAndRoundUp(OutputTextureSize.Y, FMappedCopyConstants::THREAD_GROUP_SIZE),
 		1);
 
 	FGlobalShaderMap* GlobalShaderMap = GetGlobalShaderMap(GMaxRHIFeatureLevel);
@@ -278,6 +343,35 @@ void FInputProcessBase::AddPasses(
 	}
 }
 
+bool FInputProcessBase::HasPreprocessInput(EResourceName TensorName, int32 FrameIdx) const
+{
+	check(FrameIdx == 0);
+
+	switch(TensorName)
+	{
+		case EResourceName::Color:
+		case EResourceName::Albedo:
+		case EResourceName::Normal:
+			return true;
+	}
+	return false;
+}
+
+void FInputProcessBase::PreprocessInput(
+		FRDGBuilder& GraphBuilder,
+		FRDGTextureRef Texture,
+		EResourceName TensorName,
+		int32 FrameIdx,
+		FRDGTextureRef PreprocessedTexture) const
+{
+	IOProcessBaseHelper::AddPreOrPostProcess(GraphBuilder, Texture, TensorName, FrameIdx, PreprocessedTexture);
+
+	if (TransferFunction.IsValid() && TensorName == EResourceName::Color)
+	{
+		TransferFunction->RDGForward(GraphBuilder, PreprocessedTexture, PreprocessedTexture);
+	}
+}
+
 void FInputProcessBase::WriteInputBuffer(
 	FRDGBuilder& GraphBuilder,
 	const UE::NNE::FTensorDesc& TensorDesc,
@@ -384,6 +478,13 @@ void FOutputProcessBase::AddPasses(
 	}
 }
 
+bool FOutputProcessBase::HasPostprocessOutput(EResourceName TensorName, int32 FrameIdx) const
+{
+	check(FrameIdx == 0);
+
+	return TensorName == EResourceName::Output;
+}
+
 void FOutputProcessBase::ReadOutputBuffer(
 	FRDGBuilder& GraphBuilder,
 	const UE::NNE::FTensorDesc& TensorDesc,
@@ -407,4 +508,17 @@ void FOutputProcessBase::ReadOutputBuffer(
 	IOProcessBaseHelper::AddWriteOutputPass(GraphBuilder, BufferUAV, BufferSize, TextureUAV, DataType, ChannelMapping);
 }
 
+void FOutputProcessBase::PostprocessOutput(
+		FRDGBuilder& GraphBuilder,
+		FRDGTextureRef Texture,
+		FRDGTextureRef PostprocessedTexture) const
+{
+	if (TransferFunction.IsValid())
+	{
+		TransferFunction->RDGInverse(GraphBuilder, Texture, PostprocessedTexture);
+	}
+
+	IOProcessBaseHelper::AddPreOrPostProcess(GraphBuilder, PostprocessedTexture, EResourceName::Output, 0, PostprocessedTexture);
+}
+
 } // namespace UE::NNEDenoiser::Private
\ No newline at end of file
diff --git a/Engine/Plugins/Experimental/NNEDenoiser/Source/NNEDenoiser/Private/NNEDenoiserIOProcessBase.h b/Engine/Plugins/Experimental/NNEDenoiser/Source/NNEDenoiser/Private/NNEDenoiserIOProcessBase.h
--- a/Engine/Plugins/Experimental/NNEDenoiser/Source/NNEDenoiser/Private/NNEDenoiserIOProcessBase.h
+++ b/Engine/Plugins/Experimental/NNEDenoiser/Source/NNEDenoiser/Private/NNEDenoiserIOProcessBase.h
@@ -8,10 +8,12 @@
 namespace UE::NNEDenoiser::Private
 {
 
+class ITransferFunction;
+
 class FInputProcessBase : public IInputProcess
 {
 public:
-	FInputProcessBase(FResourceMappingList InputLayout) : InputLayout(InputLayout)
+	FInputProcessBase(FResourceMappingList InputLayout, TSharedPtr<ITransferFunction> TransferFunction) : InputLayout(InputLayout), TransferFunction(TransferFunction)
 	{
 
 	}
@@ -32,20 +34,14 @@ public:
 		TConstArrayView<FRDGBufferRef> OutputBuffers) const override;
 
 protected:
-	virtual bool HasPreprocessInput(EResourceName TensorName, int32 FrameIdx) const
-	{
-		return false;
-	}
+	virtual bool HasPreprocessInput(EResourceName TensorName, int32 FrameIdx) const;
 
 	virtual void PreprocessInput(
 		FRDGBuilder& GraphBuilder,
 		FRDGTextureRef Texture,
 		EResourceName TensorName,
 		int32 FrameIdx,
-		FRDGTextureRef PreprocessedTexture) const
-	{
-		
-	}
+		FRDGTextureRef PreprocessedTexture) const;
 
 	virtual void WriteInputBuffer(
 		FRDGBuilder& GraphBuilder,
@@ -57,12 +53,13 @@ protected:
 
 private:
 	FResourceMappingList InputLayout;
+	TSharedPtr<ITransferFunction> TransferFunction;
 };
 
 class FOutputProcessBase : public IOutputProcess
 {
 public:
-	FOutputProcessBase(FResourceMappingList OutputLayout) : OutputLayout(OutputLayout)
+	FOutputProcessBase(FResourceMappingList OutputLayout, TSharedPtr<ITransferFunction> TransferFunction) : OutputLayout(OutputLayout), TransferFunction(TransferFunction)
 	{
 
 	}
@@ -80,10 +77,7 @@ public:
 		FRDGTextureRef OutputTexture) const override;
 
 protected:
-	virtual bool HasPostprocessOutput(EResourceName TensorName, int32 FrameIdx) const
-	{
-		return false;
-	}
+	virtual bool HasPostprocessOutput(EResourceName TensorName, int32 FrameIdx) const;
 
 	virtual void ReadOutputBuffer(
 		FRDGBuilder& GraphBuilder,
@@ -97,13 +91,11 @@ protected:
 	virtual void PostprocessOutput(
 		FRDGBuilder& GraphBuilder,
 		FRDGTextureRef Texture,
-		FRDGTextureRef PostprocessedTexture) const
-	{
-		
-	}
+		FRDGTextureRef PostprocessedTexture) const;
 
 private:
 	FResourceMappingList OutputLayout;
+	TSharedPtr<ITransferFunction> TransferFunction;
 };
 
 } // namespace UE::NNEDenoiser::Private
\ No newline at end of file
diff --git a/Engine/Plugins/Experimental/NNEDenoiser/Source/NNEDenoiser/Private/NNEDenoiserIOProcessOidn.cpp b/Engine/Plugins/Experimental/NNEDenoiser/Source/NNEDenoiser/Private/NNEDenoiserIOProcessOidn.cpp
deleted file mode 100644
diff --git a/Engine/Plugins/Experimental/NNEDenoiser/Source/NNEDenoiser/Private/NNEDenoiserIOProcessOidn.h b/Engine/Plugins/Experimental/NNEDenoiser/Source/NNEDenoiser/Private/NNEDenoiserIOProcessOidn.h
deleted file mode 100644
diff --git a/Engine/Plugins/Experimental/NNEDenoiser/Source/NNEDenoiser/Private/NNEDenoiserResourceManager.cpp b/Engine/Plugins/Experimental/NNEDenoiser/Source/NNEDenoiser/Private/NNEDenoiserResourceManager.cpp
--- a/Engine/Plugins/Experimental/NNEDenoiser/Source/NNEDenoiser/Private/NNEDenoiserResourceManager.cpp
+++ b/Engine/Plugins/Experimental/NNEDenoiser/Source/NNEDenoiser/Private/NNEDenoiserResourceManager.cpp
@@ -85,6 +85,7 @@ void FResourceManager::BeginTile(int32 TileIndex)
 	// Only clear if we don't overwrite the whole texture
 	const bool bNeedClear = Tile.InputOffsets.Width() != 0 || Tile.InputOffsets.Height() != 0;
 	const FLinearColor ClearColor = {.0f, 0.f, .0f, .0f};
+	const FLinearColor ClearNormal = {-1.0f, -1.0f, -1.0f, -1.0f};
 
 	for (const auto& KeyValue : TextureMap)
 	{
@@ -102,7 +103,7 @@ void FResourceManager::BeginTile(int32 TileIndex)
 		{
 			if (bNeedClear)
 			{
-				AddClearRenderTargetPass(GraphBuilder, IntermediateTextures[I], ClearColor);
+				AddClearRenderTargetPass(GraphBuilder, IntermediateTextures[I], KeyValue.Key == EResourceName::Normal ? ClearNormal : ClearColor);
 			}
 			AddCopyTexturePass(GraphBuilder, Textures[I], IntermediateTextures[I], CopyInputInfo);
 		}
diff --git a/Engine/Plugins/Experimental/NNEDenoiser/Source/NNEDenoiser/Private/NNEDenoiserTransferFunction.h b/Engine/Plugins/Experimental/NNEDenoiser/Source/NNEDenoiser/Private/NNEDenoiserTransferFunction.h
--- a/Engine/Plugins/Experimental/NNEDenoiser/Source/NNEDenoiser/Private/NNEDenoiserTransferFunction.h
+++ b/Engine/Plugins/Experimental/NNEDenoiser/Source/NNEDenoiser/Private/NNEDenoiserTransferFunction.h
@@ -2,87 +2,26 @@
 
 #pragma once
 
-#include "Math/UnrealMathUtility.h"
+#include "Containers/ContainersFwd.h"
+#include "RenderGraphFwd.h"
+
+struct FLinearColor;
 
 namespace UE::NNEDenoiser::Private
 {
-	constexpr float HDRMax = 65504.f; // maximum HDR value
-
-	// from [REDACTED_PROFILE_URL]/oidn
-	struct FPU
-	{
-		static constexpr float a  =  1.41283765e+03f;
-		static constexpr float b  =  1.64593172e+00f;
-		static constexpr float c  =  4.31384981e-01f;
-		static constexpr float d  = -2.94139609e-03f;
-		static constexpr float e  =  1.92653254e-01f;
-		static constexpr float f  =  6.26026094e-03f;
-		static constexpr float g  =  9.98620152e-01f;
-		static constexpr float y0 =  1.57945760e-06f;
-		static constexpr float y1 =  3.22087631e-02f;
-		static constexpr float x0 =  2.23151711e-03f;
-		static constexpr float x1 =  3.70974749e-01f;
-
-		static float Forward(float y)
-		{
-			if (y <= y0)
-			return a * y;
-			else if (y <= y1)
-			return b * FMath::Pow(y, c) + d;
-			else
-			return e * FMath::Loge(y + f) + g;
-		}
-
-		static float Inverse(float x)
-		{
-			if (x <= x0)
-			return x / a;
-			else if (x <= x1)
-			return FMath::Pow((x - d) / b, 1.f/c);
-			else
-			return FMath::Exp((x - g) / e) - f;
-		}
-	};
-
-	class ITransferFunction
-	{
-		public:
-			virtual float Forward(float Y) = 0;
-			virtual float Inverse(float X) = 0;
-	};
 
-	class FHDRTransferFunction : public ITransferFunction
-	{
-		public:
-			explicit FHDRTransferFunction(float Ymax)
-			{
-				InvNormScale = FPU::Forward(Ymax);
-				NormScale = 1.0f / InvNormScale;
-			}
-
-			float Forward(float Y) override
-			{
-				return FPU::Forward(Y) * NormScale;
-			}
-
-			float Inverse(float X) override
-			{
-				return FPU::Inverse(X * InvNormScale);
-			}
-
-			float GetNormScale() const
-			{
-				return NormScale;
-			}
+class ITransferFunction
+{
+public:
+	virtual ~ITransferFunction() { }
 
-			float GetInvNormScale() const
-			{
-				return InvNormScale;
-			}
+	virtual void Forward(TConstArrayView<FLinearColor> InputImage, float InputScale, TArray<FLinearColor>& OutputImage) const = 0;
+	virtual void Inverse(TConstArrayView<FLinearColor> InputImage, float InvInputScale, TArray<FLinearColor>& OutputImage) const = 0;
 
-		private:
-			float NormScale = 1.0f;
-			float InvNormScale = 1.0f;
-	};
+	virtual void RDGSetInputScale(FRDGBufferRef InputScaleBuffer) = 0;
+	
+	virtual void RDGForward(FRDGBuilder& GraphBuilder, FRDGTextureRef InputTexture, FRDGTextureRef OutputTexture) const = 0;
+	virtual void RDGInverse(FRDGBuilder& GraphBuilder, FRDGTextureRef InputTexture, FRDGTextureRef OutputTexture) const = 0;
+};
 
 }
\ No newline at end of file
diff --git a/Engine/Plugins/Experimental/NNEDenoiser/Source/NNEDenoiser/Private/NNEDenoiserTransferFunctionOidn.cpp b/Engine/Plugins/Experimental/NNEDenoiser/Source/NNEDenoiser/Private/NNEDenoiserTransferFunctionOidn.cpp
new file mode 100644
--- /dev/null
+++ b/Engine/Plugins/Experimental/NNEDenoiser/Source/NNEDenoiser/Private/NNEDenoiserTransferFunctionOidn.cpp
@@ -0,0 +1,173 @@
+// Copyright Epic Games, Inc. All Rights Reserved.
+
+#include "NNEDenoiserTransferFunctionOidn.h"
+#include "NNEDenoiserShadersTransferFunctionOidnCS.h"
+#include "Algo/Transform.h"
+
+DECLARE_GPU_STAT_NAMED(FNNEDenoiserTransferFunctionOidn, TEXT("NNEDenoiser.TransferFunctionOidn"));
+
+namespace UE::NNEDenoiser::Private::Oidn
+{
+
+namespace Helper
+{
+
+constexpr float HDRMax = 65504.f; // maximum HDR value
+
+// from [REDACTED_PROFILE_URL]/oidn
+struct FPU
+{
+	static constexpr float a  =  1.41283765e+03f;
+	static constexpr float b  =  1.64593172e+00f;
+	static constexpr float c  =  4.31384981e-01f;
+	static constexpr float d  = -2.94139609e-03f;
+	static constexpr float e  =  1.92653254e-01f;
+	static constexpr float f  =  6.26026094e-03f;
+	static constexpr float g  =  9.98620152e-01f;
+	static constexpr float y0 =  1.57945760e-06f;
+	static constexpr float y1 =  3.22087631e-02f;
+	static constexpr float x0 =  2.23151711e-03f;
+	static constexpr float x1 =  3.70974749e-01f;
+
+	static float Forward(float y)
+	{
+		if (y <= y0)
+		return a * y;
+		else if (y <= y1)
+		return b * FMath::Pow(y, c) + d;
+		else
+		return e * FMath::Loge(y + f) + g;
+	}
+
+	static float Inverse(float x)
+	{
+		if (x <= x0)
+		return x / a;
+		else if (x <= x1)
+		return FMath::Pow((x - d) / b, 1.f/c);
+		else
+		return FMath::Exp((x - g) / e) - f;
+	}
+};
+
+void ApplyTransferFunction(
+	FRDGBuilder& GraphBuilder,
+	FRDGTextureRef InputTexture,
+	FRDGTextureRef OutputTexture,
+	FRDGBufferRef InputScaleBuffer,
+	float NormScale,
+	float InvNormScale,
+	bool Forward)
+{
+	using namespace UE::NNEDenoiserShaders::Internal;
+
+	const FIntVector InputTextureSize = InputTexture->Desc.GetSize();
+	const FIntVector OutputTextureSize = OutputTexture->Desc.GetSize();
+
+	FTransferFunctionOidnCS::FParameters *ShaderParameters = GraphBuilder.AllocParameters<FTransferFunctionOidnCS::FParameters>();
+	ShaderParameters->InputTextureWidth = InputTextureSize.X;
+	ShaderParameters->InputTextureHeight = InputTextureSize.Y;
+	ShaderParameters->InputTexture = InputTexture;
+	ShaderParameters->OutputTextureWidth = OutputTextureSize.X;
+	ShaderParameters->OutputTextureHeight = OutputTextureSize.Y;
+	ShaderParameters->OutputTexture = GraphBuilder.CreateUAV(FRDGTextureUAVDesc(OutputTexture));
+	ShaderParameters->InputScaleBuffer = GraphBuilder.CreateUAV(FRDGBufferUAVDesc(InputScaleBuffer, EPixelFormat::PF_R32_FLOAT));
+	ShaderParameters->NormScale = NormScale;
+	ShaderParameters->InvNormScale = InvNormScale;
+
+	FTransferFunctionOidnCS::FPermutationDomain PermutationVector;
+	PermutationVector.Set<FTransferFunctionOidnCS::FTransferFunctionOidnMode>(Forward ? ETransferFunctionOidnMode::Forwward : ETransferFunctionOidnMode::Inverse);
+
+	FIntVector ThreadGroupCount = FIntVector(
+		FMath::DivideAndRoundUp(OutputTextureSize.X, FTransferFunctionOidnConstants::THREAD_GROUP_SIZE),
+		FMath::DivideAndRoundUp(OutputTextureSize.Y, FTransferFunctionOidnConstants::THREAD_GROUP_SIZE),
+		1);
+
+	FGlobalShaderMap* GlobalShaderMap = GetGlobalShaderMap(GMaxRHIFeatureLevel);
+	TShaderMapRef<FTransferFunctionOidnCS> Shader(GlobalShaderMap, PermutationVector);
+
+	RDG_EVENT_SCOPE(GraphBuilder, "NNEDenoiser.TransferFunctionOidn");
+	RDG_GPU_STAT_SCOPE(GraphBuilder, FNNEDenoiserTransferFunctionOidn);
+
+	FComputeShaderUtils::AddPass(
+		GraphBuilder,
+		RDG_EVENT_NAME("NNEDenoiser.TransferFunctionOidn"),
+		ERDGPassFlags::Compute | ERDGPassFlags::NeverCull,
+		Shader,
+		ShaderParameters,
+		ThreadGroupCount);
+}
+
+} // Helpser
+
+FTransferFunction::FTransferFunction()
+{
+	InvNormScale = Helper::FPU::Forward(Helper::HDRMax);
+	NormScale = 1.0f / InvNormScale;
+}
+
+void FTransferFunction::Forward(TConstArrayView<FLinearColor> InputImage, float InputScale, TArray<FLinearColor>& OutputImage) const
+{
+	auto InputProcess = [this, InputScale] (float Value)
+	{
+		Value *= InputScale;
+		Value = FMath::Clamp(Value, 0.0f, FLT_MAX);
+		
+		return Helper::FPU::Forward(Value) * NormScale;
+	};
+
+	Algo::Transform(InputImage, OutputImage, [&InputProcess] (const FLinearColor& Color)
+	{
+		return FLinearColor(
+			InputProcess(Color.R),
+			InputProcess(Color.G),
+			InputProcess(Color.B),
+			Color.A
+		);
+	});
+}
+
+void FTransferFunction::Inverse(TConstArrayView<FLinearColor> InputImage, float InvInputScale, TArray<FLinearColor>& OutputImage) const
+{
+	auto OutputProcess = [this, InvInputScale] (float Value)
+	{
+		Value = FMath::Clamp(Value, 0.0f, FLT_MAX);
+		Value = Helper::FPU::Inverse(Value * InvNormScale);
+		Value *= InvInputScale;
+
+		return Value;
+	};
+
+	Algo::Transform(InputImage, OutputImage, [&OutputProcess] (const FLinearColor& Color)
+	{
+		return FLinearColor(
+			OutputProcess(Color.R),
+			OutputProcess(Color.G),
+			OutputProcess(Color.B),
+			Color.A
+		);
+	});
+}
+
+void FTransferFunction::RDGSetInputScale(FRDGBufferRef InInputScaleBuffer)
+{
+	check(InInputScaleBuffer);
+
+	InputScaleBuffer = InInputScaleBuffer;
+}
+
+void FTransferFunction::RDGForward(FRDGBuilder& GraphBuilder, FRDGTextureRef InputTexture, FRDGTextureRef OutputTexture) const
+{
+	check(InputScaleBuffer);
+
+	Helper::ApplyTransferFunction(GraphBuilder, InputTexture, OutputTexture, InputScaleBuffer, NormScale, InvNormScale, true);
+}
+
+void FTransferFunction::RDGInverse(FRDGBuilder& GraphBuilder, FRDGTextureRef InputTexture, FRDGTextureRef OutputTexture) const
+{
+	check(InputScaleBuffer);
+	
+	Helper::ApplyTransferFunction(GraphBuilder, InputTexture, OutputTexture, InputScaleBuffer, NormScale, InvNormScale, false);
+}
+
+} // namespace UE::NNEDenoiser::Private::Oidn
\ No newline at end of file
diff --git a/Engine/Plugins/Experimental/NNEDenoiser/Source/NNEDenoiser/Private/NNEDenoiserTransferFunctionOidn.h b/Engine/Plugins/Experimental/NNEDenoiser/Source/NNEDenoiser/Private/NNEDenoiserTransferFunctionOidn.h
new file mode 100644
--- /dev/null
+++ b/Engine/Plugins/Experimental/NNEDenoiser/Source/NNEDenoiser/Private/NNEDenoiserTransferFunctionOidn.h
@@ -0,0 +1,30 @@
+// Copyright Epic Games, Inc. All Rights Reserved.
+
+#pragma once
+
+#include "NNEDenoiserTransferFunction.h"
+
+namespace UE::NNEDenoiser::Private::Oidn
+{
+
+class FTransferFunction : public ITransferFunction
+{
+public:
+	FTransferFunction();
+	virtual ~FTransferFunction() { }
+
+	virtual void Forward(TConstArrayView<FLinearColor> InputImage, float InputScale, TArray<FLinearColor>& OutputImage) const override;
+	virtual void Inverse(TConstArrayView<FLinearColor> InputImage, float InvInputScale, TArray<FLinearColor>& OutputImage) const override;
+
+	virtual void RDGSetInputScale(FRDGBufferRef InInputScaleBuffer) override;
+	
+	virtual void RDGForward(FRDGBuilder& GraphBuilder, FRDGTextureRef InputTexture, FRDGTextureRef OutputTexture) const override;
+	virtual void RDGInverse(FRDGBuilder& GraphBuilder, FRDGTextureRef InputTexture, FRDGTextureRef OutputTexture) const override;
+
+private:
+	float NormScale = 1.0f;
+	float InvNormScale = 1.0f;
+	FRDGBufferRef InputScaleBuffer{};
+};
+
+} // UE::NNEDenoiser::Private::Oidn
\ No newline at end of file
diff --git a/Engine/Plugins/Experimental/NNEDenoiser/Source/NNEDenoiser/Private/NNEDenoiserUnitTestOidn.cpp b/Engine/Plugins/Experimental/NNEDenoiser/Source/NNEDenoiser/Private/NNEDenoiserUnitTestOidn.cpp
new file mode 100644
--- /dev/null
+++ b/Engine/Plugins/Experimental/NNEDenoiser/Source/NNEDenoiser/Private/NNEDenoiserUnitTestOidn.cpp
@@ -0,0 +1,145 @@
+// Copyright Epic Games, Inc. All Rights Reserved.
+
+#include "CoreMinimal.h"
+#include "NNEDenoiserTransferFunctionOidn.h"
+#include "NNEDenoiserUtils.h"
+#include "NNERuntimeCPU.h"
+#include "NNERuntimeRDG.h"
+#include "RenderGraphBuilder.h"
+
+#if WITH_DEV_AUTOMATION_TESTS
+#include "Misc/AutomationTest.h"
+#endif
+
+namespace UE::NNEDenoiser::Private
+{
+
+#if WITH_DEV_AUTOMATION_TESTS
+
+BEGIN_SHADER_PARAMETER_STRUCT(FTransferFunctionTestUploadTextureParameters, )
+	RDG_TEXTURE_ACCESS(InputTexture, ERHIAccess::CopyDest)
+END_SHADER_PARAMETER_STRUCT()
+
+BEGIN_SHADER_PARAMETER_STRUCT(FTransferFunctionTestUploadBufferParameters, )
+	RDG_BUFFER_ACCESS(InputBuffer, ERHIAccess::CopyDest)
+END_SHADER_PARAMETER_STRUCT()
+
+BEGIN_SHADER_PARAMETER_STRUCT(FTransferFunctionTestDownloadParameters, )
+	RDG_TEXTURE_ACCESS(OutputTexture, ERHIAccess::CopySrc)
+END_SHADER_PARAMETER_STRUCT()
+
+IMPLEMENT_SIMPLE_AUTOMATION_TEST(FTransferFunctionTest, "NNEDenoiser.UnitTests.OIDN.TransferFunction.Forward", EAutomationTestFlags::ApplicationContextMask | EAutomationTestFlags::FeatureMask | EAutomationTestFlags::EngineFilter)
+bool FTransferFunctionTest::RunTest(const FString& Parameter)
+{
+	const int32 Width = 1920;
+	const int32 Height = 1088;
+	const int32 Num = Width * Height;
+
+	const TArray<float> InputScale = {0.5f, 2.0f};
+
+	TArray<FLinearColor> TestData;
+	TestData.SetNumUninitialized(Num);
+
+	for (int i = 0; i < TestData.Num(); i++)
+	{
+		TestData[i] = FLinearColor(FMath::FRand(), FMath::FRand(), FMath::FRand(), 1.0);
+	}
+
+	Oidn::FTransferFunction TransferFunction;
+
+	TArray<FLinearColor> ResultCPU;
+	TransferFunction.Forward(TestData, InputScale[0], ResultCPU);
+
+	TArray<FLinearColor> ResultRDG;
+	ResultRDG.SetNumZeroed(Num);
+
+	FEvent* Signal = FGenericPlatformProcess::GetSynchEventFromPool(true);
+
+	ENQUEUE_RENDER_COMMAND(NNETransferFunctionTest)
+	(
+		[&TestData, &ResultRDG, &InputScale, TransferFunction, Num, Signal](FRHICommandListImmediate& RHICmdList) mutable
+		{
+			ERHIPipeline Pipeline = RHICmdList.GetPipeline();
+			if (Pipeline == ERHIPipeline::None)
+			{
+				RHICmdList.SwitchPipeline(ERHIPipeline::Graphics);
+			}
+
+			FRDGBuilder	GraphBuilder(RHICmdList);
+
+			const FRDGTextureDesc InputTextureDesc = FRDGTextureDesc::Create2D({Width, Height}, EPixelFormat::PF_A32B32G32R32F, FClearValueBinding::Black, ETextureCreateFlags::None);
+			FRDGTextureRef InputTexture = GraphBuilder.CreateTexture(InputTextureDesc, TEXT("TransferFunctionInputTexture"));
+
+			{
+				FTransferFunctionTestUploadTextureParameters* Parameters = GraphBuilder.AllocParameters<FTransferFunctionTestUploadTextureParameters>();
+				Parameters->InputTexture = InputTexture;
+
+				GraphBuilder.AddPass(RDG_EVENT_NAME("NNEDenoiser.TransferFunctionTest.UploadInput"), Parameters, ERDGPassFlags::Readback,
+				[&TestData, Parameters](FRHICommandListImmediate& RHICmdList)
+				{
+					SCOPED_NAMED_EVENT_TEXT("FTransferFunctionTest::UploadInput", FColor::Magenta);
+
+					CopyTextureFromCPUToGPU(RHICmdList, TestData, {Width, Height}, Parameters->InputTexture->GetRHI());
+				});
+			}
+
+			FRDGBufferDesc InputScaleBufferDesc = FRDGBufferDesc::CreateBufferDesc(sizeof(float), 2);
+			FRDGBufferRef InputScaleBuffer = GraphBuilder.CreateBuffer(InputScaleBufferDesc, TEXT("TransferFunctionInputScaleBuffer"));
+
+			{
+				FTransferFunctionTestUploadBufferParameters* Parameters = GraphBuilder.AllocParameters<FTransferFunctionTestUploadBufferParameters>();
+				Parameters->InputBuffer = InputScaleBuffer;
+
+				GraphBuilder.AddPass(RDG_EVENT_NAME("NNEDenoiser.TransferFunctionTest.UploadScale"), Parameters, ERDGPassFlags::Readback,
+				[&InputScale, Parameters](FRHICommandListImmediate& RHICmdList)
+				{
+					SCOPED_NAMED_EVENT_TEXT("FTransferFunctionTest::UploadScale", FColor::Magenta);
+
+					CopyBufferFromCPUToGPU(RHICmdList, InputScale, 2, Parameters->InputBuffer->GetRHI());
+				});
+			}
+
+			const FRDGTextureDesc OutputTextureDesc = FRDGTextureDesc::Create2D({Width, Height}, EPixelFormat::PF_A32B32G32R32F, FClearValueBinding::Black, ETextureCreateFlags::UAV);
+			FRDGTextureRef OutputTexture = GraphBuilder.CreateTexture(OutputTextureDesc, TEXT("TransferFunctionOutputTexture"));
+
+			TransferFunction.RDGSetInputScale(InputScaleBuffer);
+
+			TransferFunction.RDGForward(GraphBuilder, InputTexture, OutputTexture);
+
+			{
+				FTransferFunctionTestDownloadParameters* Parameters = GraphBuilder.AllocParameters<FTransferFunctionTestDownloadParameters>();
+				Parameters->OutputTexture = OutputTexture;
+
+				GraphBuilder.AddPass(RDG_EVENT_NAME("NNEDenoiser.TransferFunctionTest.Download"), Parameters, ERDGPassFlags::Readback,
+				[&ResultRDG, Num, Parameters](FRHICommandListImmediate& RHICmdList)
+				{
+					SCOPED_NAMED_EVENT_TEXT("FTransferFunctionTest::Download", FColor::Magenta);
+
+					CopyTextureFromGPUToCPU(RHICmdList, Parameters->OutputTexture->GetRHI(), {Width, Height}, ResultRDG);
+				});
+			}
+
+			GraphBuilder.Execute();
+
+			RHICmdList.BlockUntilGPUIdle();
+
+			Signal->Trigger();
+		}
+	);
+
+	Signal->Wait();
+
+	for (int32 i = 0; i < Num; i++)
+	{
+		UTEST_EQUAL_TOLERANCE(TEXT("TransferFunction"), ResultRDG[i].R, ResultCPU[i].R, 1e-6f);
+		UTEST_EQUAL_TOLERANCE(TEXT("TransferFunction"), ResultRDG[i].G, ResultCPU[i].G, 1e-6f);
+		UTEST_EQUAL_TOLERANCE(TEXT("TransferFunction"), ResultRDG[i].B, ResultCPU[i].B, 1e-6f);
+		UTEST_EQUAL(TEXT("TransferFunction"), ResultRDG[i].A, ResultCPU[i].A);
+	}
+
+	return true;
+}
+
+#endif //WITH_DEV_AUTOMATION_TESTS
+
+} // namespace UE::NNEDenoiser::Internal
diff --git a/Engine/Plugins/Experimental/NNEDenoiser/Source/NNEDenoiser/Private/NNEDenoiserUnitTests.cpp b/Engine/Plugins/Experimental/NNEDenoiser/Source/NNEDenoiser/Private/NNEDenoiserUnitTests.cpp
new file mode 100644
--- /dev/null
+++ b/Engine/Plugins/Experimental/NNEDenoiser/Source/NNEDenoiser/Private/NNEDenoiserUnitTests.cpp
@@ -0,0 +1,114 @@
+// Copyright Epic Games, Inc. All Rights Reserved.
+
+#include "CoreMinimal.h"
+#include "NNEDenoiserAutoExposure.h"
+#include "NNEDenoiserUtils.h"
+#include "NNERuntimeCPU.h"
+#include "NNERuntimeRDG.h"
+#include "RenderGraphBuilder.h"
+
+#if WITH_DEV_AUTOMATION_TESTS
+#include "Misc/AutomationTest.h"
+#endif
+
+namespace UE::NNEDenoiser::Private
+{
+
+#if WITH_DEV_AUTOMATION_TESTS
+
+BEGIN_SHADER_PARAMETER_STRUCT(FAutoExposureTestUploadParameters, )
+	RDG_TEXTURE_ACCESS(InputTexture, ERHIAccess::CopyDest)
+END_SHADER_PARAMETER_STRUCT()
+
+BEGIN_SHADER_PARAMETER_STRUCT(FAutoExposureTestDownloadParameters, )
+	RDG_BUFFER_ACCESS(OutputBuffer, ERHIAccess::CopySrc)
+END_SHADER_PARAMETER_STRUCT()
+
+IMPLEMENT_SIMPLE_AUTOMATION_TEST(FAutoExposureTest, "NNEDenoiser.UnitTests.AutoExposure", EAutomationTestFlags::ApplicationContextMask | EAutomationTestFlags::FeatureMask | EAutomationTestFlags::EngineFilter)
+bool FAutoExposureTest::RunTest(const FString& Parameter)
+{
+	const int32 Width = 1920;
+	const int32 Height = 1088;
+
+	TArray<FLinearColor> TestData;
+	TestData.SetNumUninitialized(Width * Height);
+
+	for (int i = 0; i < TestData.Num(); i++)
+	{
+		TestData[i] = FLinearColor(FMath::FRand(), FMath::FRand(), FMath::FRand(), 1.0);
+	}
+
+	FAutoExposure AutoExposure;
+
+	float ResultCPU = 0.0f;
+	AutoExposure.Run(TestData, {Width, Height}, ResultCPU);
+
+	TArray<float> ResultRDG;
+	ResultRDG.SetNumZeroed(2);
+
+	FEvent* Signal = FGenericPlatformProcess::GetSynchEventFromPool(true);
+
+	ENQUEUE_RENDER_COMMAND(NNEAutoExposureTest)
+	(
+		[&TestData, &ResultRDG, AutoExposure, Signal](FRHICommandListImmediate& RHICmdList) mutable
+		{
+			ERHIPipeline Pipeline = RHICmdList.GetPipeline();
+			if (Pipeline == ERHIPipeline::None)
+			{
+				RHICmdList.SwitchPipeline(ERHIPipeline::Graphics);
+			}
+
+			FRDGBuilder	GraphBuilder(RHICmdList);
+
+			const FRDGTextureDesc TextureDesc = FRDGTextureDesc::Create2D({Width, Height}, EPixelFormat::PF_A32B32G32R32F, FClearValueBinding::Black, ETextureCreateFlags::None);
+			FRDGTextureRef InputTexture = GraphBuilder.CreateTexture(TextureDesc, TEXT("AutoExposureInputTexture"));
+
+			{
+				FAutoExposureTestUploadParameters* Parameters = GraphBuilder.AllocParameters<FAutoExposureTestUploadParameters>();
+				Parameters->InputTexture = InputTexture;
+
+				GraphBuilder.AddPass(RDG_EVENT_NAME("NNEDenoiser.AutoExposureTest.Upload"), Parameters, ERDGPassFlags::Readback,
+				[&TestData, Parameters](FRHICommandListImmediate& RHICmdList)
+				{
+					SCOPED_NAMED_EVENT_TEXT("FAutoExposureTest::Upload", FColor::Magenta);
+
+					CopyTextureFromCPUToGPU(RHICmdList, TestData, {Width, Height}, Parameters->InputTexture->GetRHI());
+				});
+			}
+
+			FRDGBufferDesc InputBufferDesc = FRDGBufferDesc::CreateBufferDesc(sizeof(float), 2);
+			FRDGBufferRef OutputBuffer = GraphBuilder.CreateBuffer(InputBufferDesc, TEXT("AutoExposureOutputBuffer"));
+
+			AutoExposure.EnqueueRDG(GraphBuilder, InputTexture, OutputBuffer);
+
+			{
+				FAutoExposureTestDownloadParameters* Parameters = GraphBuilder.AllocParameters<FAutoExposureTestDownloadParameters>();
+				Parameters->OutputBuffer = OutputBuffer;
+
+				GraphBuilder.AddPass(RDG_EVENT_NAME("NNEDenoiser.AutoExposureTest.Download"), Parameters, ERDGPassFlags::Readback,
+				[&ResultRDG, Parameters](FRHICommandListImmediate& RHICmdList)
+				{
+					SCOPED_NAMED_EVENT_TEXT("FAutoExposureTest::Download", FColor::Magenta);
+
+					CopyBufferFromGPUToCPU(RHICmdList, Parameters->OutputBuffer->GetRHI(), 2, ResultRDG);
+				});
+			}
+
+			GraphBuilder.Execute();
+
+			RHICmdList.BlockUntilGPUIdle();
+
+			Signal->Trigger();
+		}
+	);
+
+	Signal->Wait();
+
+	UTEST_EQUAL(TEXT("AutoExposure"), ResultRDG[0], ResultCPU);
+
+	return true;
+}
+
+#endif //WITH_DEV_AUTOMATION_TESTS
+
+} // namespace UE::NNEDenoiser::Internal
diff --git a/Engine/Plugins/Experimental/NNEDenoiser/Source/NNEDenoiser/Private/NNEDenoiserUtils.h b/Engine/Plugins/Experimental/NNEDenoiser/Source/NNEDenoiser/Private/NNEDenoiserUtils.h
--- a/Engine/Plugins/Experimental/NNEDenoiser/Source/NNEDenoiser/Private/NNEDenoiserUtils.h
+++ b/Engine/Plugins/Experimental/NNEDenoiser/Source/NNEDenoiser/Private/NNEDenoiserUtils.h
@@ -5,8 +5,7 @@
 #include "NNEDenoiserLog.h"
 #include "NNEDenoiserModelIOMappingData.h"
 #include "NNEDenoiserParameters.h"
-#include "NNEDenoiserShadersDefaultCS.h"
-#include "NNEDenoiserShadersOidnCS.h"
+#include "NNEDenoiserShadersMappedCopyCS.h"
 #include "NNETypes.h"
 #include "RHICommandList.h"
 #include "RHIResources.h"
@@ -74,9 +73,9 @@ namespace UE::NNEDenoiser::Private
 		SCOPED_NAMED_EVENT_TEXT("NNEDenoiser.CopyTextureFromCPUToGPU", FColor::Magenta);
 
 		uint32 DestStride;
-		FLinearColor* DstBuffer = static_cast<PixelType*>(RHICmdList.LockTexture2D(DstTexture, 0, RLM_WriteOnly, DestStride, false));
+		PixelType* DstBuffer = static_cast<PixelType*>(RHICmdList.LockTexture2D(DstTexture, 0, RLM_WriteOnly, DestStride, false));
 		DestStride /= sizeof(PixelType);
-		const FLinearColor* SrcBuffer = SrcArray.GetData();
+		const PixelType* SrcBuffer = SrcArray.GetData();
 		for (int32 Y = 0; Y < Size.Y; Y++, SrcBuffer += Size.X, DstBuffer += DestStride)
 		{
 			FPlatformMemory::Memcpy(DstBuffer, SrcBuffer, Size.X * sizeof(PixelType));
@@ -84,26 +83,28 @@ namespace UE::NNEDenoiser::Private
 		RHICmdList.UnlockTexture2D(DstTexture, 0, false);
 	}
 
-	inline void CopyBufferFromGPUToCPU(FRHICommandListImmediate& RHICmdList, FRHIBuffer* SrcBuffer, int32 Count, TArray<uint8>& DstArray)
+	template <typename ElementType>
+	void CopyBufferFromGPUToCPU(FRHICommandListImmediate& RHICmdList, FRHIBuffer* SrcBuffer, int32 Count, TArray<ElementType>& DstArray)
 	{
 		SCOPED_NAMED_EVENT_TEXT("NNEDenoiser.CopyBufferFromGPUToCPU", FColor::Magenta);
 
-		const uint8* Src = static_cast<uint8*>(RHICmdList.LockBuffer(SrcBuffer, 0, Count, RLM_ReadOnly));
-		uint8* Dst = DstArray.GetData();
+		const ElementType* Src = static_cast<ElementType*>(RHICmdList.LockBuffer(SrcBuffer, 0, Count * sizeof(ElementType), RLM_ReadOnly));
+		ElementType* Dst = DstArray.GetData();
 		
-		FPlatformMemory::Memcpy(Dst, Src, Count);
+		FPlatformMemory::Memcpy(Dst, Src, Count * sizeof(ElementType));
 
 		RHICmdList.UnlockBuffer(SrcBuffer);
 	}
 
-	inline void CopyBufferFromCPUToGPU(FRHICommandListImmediate& RHICmdList, const TArray<uint8>& SrcArray, int32 Count, FRHIBuffer* DstBuffer)
+	template <typename ElementType>
+	void CopyBufferFromCPUToGPU(FRHICommandListImmediate& RHICmdList, const TArray<ElementType>& SrcArray, int32 Count, FRHIBuffer* DstBuffer)
 	{
 		SCOPED_NAMED_EVENT_TEXT("NNEDenoiser.CopyBufferFromCPUToGPU", FColor::Magenta);
 
-		uint8* Dst = static_cast<uint8*>(RHICmdList.LockBuffer(DstBuffer, 0, Count, RLM_WriteOnly));
-		const uint8* Src = SrcArray.GetData();
+		ElementType* Dst = static_cast<ElementType*>(RHICmdList.LockBuffer(DstBuffer, 0, Count * sizeof(ElementType), RLM_WriteOnly));
+		const ElementType* Src = SrcArray.GetData();
 		
-		FPlatformMemory::Memcpy(Dst, Src, Count);
+		FPlatformMemory::Memcpy(Dst, Src, Count * sizeof(ElementType));
 
 		RHICmdList.UnlockBuffer(DstBuffer);
 	}
@@ -130,23 +131,4 @@ namespace UE::NNEDenoiser::Private
 		return ENNEDenoiserDataType::None;
 	}
 
-	inline NNEDenoiserShaders::Internal::ENNEDenoiserInputKind GetInputKind(EResourceName TensorName)
-	{
-		using NNEDenoiserShaders::Internal::ENNEDenoiserInputKind;
-		
-		switch(TensorName)
-		{
-			case EResourceName::Color:	return ENNEDenoiserInputKind::Color;
-			case EResourceName::Albedo:	return ENNEDenoiserInputKind::Albedo;
-			case EResourceName::Normal:	return ENNEDenoiserInputKind::Normal;
-			case EResourceName::Flow:	return ENNEDenoiserInputKind::Flow;
-			case EResourceName::Output:	return ENNEDenoiserInputKind::Output;
-		}
-
-		// There should be a case for every resource name
-		checkNoEntry();
-		
-		return ENNEDenoiserInputKind::Color;
-	}
-
 }
\ No newline at end of file
diff --git a/Engine/Plugins/Experimental/NNEDenoiser/Source/NNEDenoiser/Private/NNEDenoiserViewExtension.cpp b/Engine/Plugins/Experimental/NNEDenoiser/Source/NNEDenoiser/Private/NNEDenoiserViewExtension.cpp
--- a/Engine/Plugins/Experimental/NNEDenoiser/Source/NNEDenoiser/Private/NNEDenoiserViewExtension.cpp
+++ b/Engine/Plugins/Experimental/NNEDenoiser/Source/NNEDenoiser/Private/NNEDenoiserViewExtension.cpp
@@ -2,9 +2,9 @@
 
 #include "NNEDenoiserViewExtension.h"
 #include "NNE.h"
+#include "NNEDenoiserAutoExposure.h"
 #include "NNEDenoiserGenericDenoiser.h"
 #include "NNEDenoiserIOProcessBase.h"
-#include "NNEDenoiserIOProcessOidn.h"
 #include "NNEDenoiserLog.h"
 #include "NNEDenoiserModelData.h"
 #include "NNEDenoiserModelInstanceCPU.h"
@@ -13,6 +13,7 @@
 #include "NNEDenoiserModelIOMappingData.h"
 #include "NNEDenoiserPathTracingDenoiser.h"
 #include "NNEDenoiserPathTracingSpatialTemporalDenoiser.h"
+#include "NNEDenoiserTransferFunctionOidn.h"
 #include "NNEDenoiserUtils.h"
 #include "NNEModelData.h"
 #include "PathTracingDenoiser.h"
@@ -99,8 +100,15 @@ FString GetDenoiserModelDataNameFromCVarAndSettings(const UNNEDenoiserSettings*
 	return FString();
 }
 
-TUniquePtr<FGenericDenoiser> CreateNNEDenoiser(UNNEModelData& ModelData, EDenoiserRuntimeType RuntimeType, const FString& RuntimeNameOverride,
-	TUniquePtr<IInputProcess> InputProcess, TUniquePtr<IOutputProcess> OutputProcess, FParameters Parameters)
+TUniquePtr<FGenericDenoiser> CreateNNEDenoiser(
+	UNNEModelData& ModelData,
+	EDenoiserRuntimeType RuntimeType,
+	const FString& RuntimeNameOverride,
+	TUniquePtr<IInputProcess> InputProcess,
+	TUniquePtr<IOutputProcess> OutputProcess,
+	FParameters Parameters,
+	TUniquePtr<IAutoExposure> AutoExposure,
+	TSharedPtr<ITransferFunction> TransferFunction)
 {
 	TUniquePtr<IModelInstance> ModelInstance;
 	if (RuntimeType == EDenoiserRuntimeType::CPU)
@@ -122,7 +130,13 @@ TUniquePtr<FGenericDenoiser> CreateNNEDenoiser(UNNEModelData& ModelData, EDenois
 		return {};
 	}
 
-	return MakeUnique<FGenericDenoiser>(MoveTemp(ModelInstance), MoveTemp(InputProcess), MoveTemp(OutputProcess), MoveTemp(Parameters));
+	return MakeUnique<FGenericDenoiser>(
+		MoveTemp(ModelInstance),
+		MoveTemp(InputProcess),
+		MoveTemp(OutputProcess),
+		MoveTemp(Parameters),
+		MoveTemp(AutoExposure),
+		TransferFunction);
 }
 
 static FResourceMappingList MakeTensorLayout(UDataTable* DataTable)
@@ -263,29 +277,28 @@ TUniquePtr<FGenericDenoiser> CreateNNEDenoiserFromAsset(const FString& AssetName
 
 	check(!bIsOidnModel || !bIsInHouseModel);
 
-	TUniquePtr<IInputProcess> InputProcess;
+	TUniquePtr<FAutoExposure> AutoExposure;
+	TSharedPtr<ITransferFunction> TransferFunction;
 	if (bIsOidnModel)
 	{
-		InputProcess = MakeUnique<FInputProcessOidn>(MoveTemp(InputLayout));
-	}
-	else
-	{
-		InputProcess = MakeUnique<FInputProcessBase>(MoveTemp(InputLayout));
+		AutoExposure = MakeUnique<FAutoExposure>();
+		TransferFunction = MakeShared<Oidn::FTransferFunction>();
 	}
 
-	TUniquePtr<IOutputProcess> OutputProcess;
-	if (bIsOidnModel)
-	{
-		OutputProcess = MakeUnique<FOutputProcessOidn>(MoveTemp(OutputLayout));
-	}
-	else
-	{
-		OutputProcess = MakeUnique<FOutputProcessBase>(MoveTemp(OutputLayout));
-	}
+	TUniquePtr<IInputProcess> InputProcess = MakeUnique<FInputProcessBase>(MoveTemp(InputLayout), TransferFunction);
+	TUniquePtr<IOutputProcess> OutputProcess = MakeUnique<FOutputProcessBase>(MoveTemp(OutputLayout), TransferFunction);
 
 	FParameters Parameters = GetParametersValidated(*DenoiserModelData);
 
-	return CreateNNEDenoiser(*ModelData, RuntimeType, RuntimeNameOverride, MoveTemp(InputProcess), MoveTemp(OutputProcess), MoveTemp(Parameters));
+	return CreateNNEDenoiser(
+		*ModelData,
+		RuntimeType,
+		RuntimeNameOverride,
+		MoveTemp(InputProcess),
+		MoveTemp(OutputProcess),
+		MoveTemp(Parameters),
+		MoveTemp(AutoExposure),
+		TransferFunction);
 }
 
 FViewExtension::FViewExtension(const FAutoRegister& AutoRegister) : FSceneViewExtensionBase(AutoRegister)
diff --git a/Engine/Plugins/Experimental/NNEDenoiser/Source/NNEDenoiserShaders/Internal/NNEDenoiserShadersAutoExposureCS.h b/Engine/Plugins/Experimental/NNEDenoiser/Source/NNEDenoiserShaders/Internal/NNEDenoiserShadersAutoExposureCS.h
new file mode 100644
--- /dev/null
+++ b/Engine/Plugins/Experimental/NNEDenoiser/Source/NNEDenoiserShaders/Internal/NNEDenoiserShadersAutoExposureCS.h
@@ -0,0 +1,98 @@
+// Copyright Epic Games, Inc. All Rights Reserved.
+
+#pragma once
+
+#include "GlobalShader.h"
+#include "RenderGraphUtils.h"
+#include "ShaderParameterUtils.h"
+
+namespace UE::NNEDenoiserShaders::Internal
+{
+	
+	class NNEDENOISERSHADERS_API FAutoExposureCS : public FGlobalShader
+	{
+		DECLARE_GLOBAL_SHADER(FAutoExposureCS);
+		SHADER_USE_PARAMETER_STRUCT(FAutoExposureCS, FGlobalShader)
+
+	public:
+
+		BEGIN_SHADER_PARAMETER_STRUCT(FParameters, )
+			SHADER_PARAMETER(int32, InputTextureWidth)
+			SHADER_PARAMETER(int32, InputTextureHeight)
+			SHADER_PARAMETER_RDG_TEXTURE(Texture2D, InputTexture)
+			SHADER_PARAMETER_RDG_BUFFER_UAV(RWBuffer, OutputBuffer)
+		END_SHADER_PARAMETER_STRUCT()
+
+		static void ModifyCompilationEnvironment(const FGlobalShaderPermutationParameters& InParameters, FShaderCompilerEnvironment& OutEnvironment);
+	};
+
+	class FAutoExposureDownsampleConstants
+	{
+	public:
+		static constexpr int32 MAX_BIN_SIZE{16};
+		static constexpr int32 THREAD_GROUP_SIZE{MAX_BIN_SIZE};
+	};
+
+	class NNEDENOISERSHADERS_API FAutoExposureDownsampleCS : public FGlobalShader
+	{
+		DECLARE_GLOBAL_SHADER(FAutoExposureDownsampleCS);
+		SHADER_USE_PARAMETER_STRUCT(FAutoExposureDownsampleCS, FGlobalShader)
+
+	public:
+
+		BEGIN_SHADER_PARAMETER_STRUCT(FParameters, )
+			SHADER_PARAMETER(int32, InputTextureWidth)
+			SHADER_PARAMETER(int32, InputTextureHeight)
+			SHADER_PARAMETER_RDG_TEXTURE(Texture2D, InputTexture)
+			SHADER_PARAMETER(int32, NumBinsW)
+			SHADER_PARAMETER(int32, NumBinsH)
+			SHADER_PARAMETER_RDG_BUFFER_UAV(RWBuffer, OutputBins)
+		END_SHADER_PARAMETER_STRUCT()
+
+		static void ModifyCompilationEnvironment(const FGlobalShaderPermutationParameters& InParameters, FShaderCompilerEnvironment& OutEnvironment);
+	};
+
+	class FAutoExposureReduceConstants
+	{
+	public:
+		static constexpr int32 THREAD_GROUP_SIZE{32};
+		static constexpr float EPS{1e-8f};
+		static constexpr float KEY{0.18f};
+	};
+
+	class NNEDENOISERSHADERS_API FAutoExposureReduceCS : public FGlobalShader
+	{
+		DECLARE_GLOBAL_SHADER(FAutoExposureReduceCS);
+		SHADER_USE_PARAMETER_STRUCT(FAutoExposureReduceCS, FGlobalShader)
+
+	public:
+
+		BEGIN_SHADER_PARAMETER_STRUCT(FParameters, )
+			SHADER_PARAMETER(int32, InputSize)
+			SHADER_PARAMETER_RDG_BUFFER_UAV(RWBuffer, InputBins)
+			SHADER_PARAMETER_RDG_BUFFER_UAV(RWBuffer, OutputSums)
+			SHADER_PARAMETER_RDG_BUFFER_UAV(RWBuffer, OutputCounts)
+			SHADER_PARAMETER(int32, NumThreads)
+		END_SHADER_PARAMETER_STRUCT()
+
+		static void ModifyCompilationEnvironment(const FGlobalShaderPermutationParameters& InParameters, FShaderCompilerEnvironment& OutEnvironment);
+	};
+
+	class NNEDENOISERSHADERS_API FAutoExposureReduceFinalCS : public FGlobalShader
+	{
+		DECLARE_GLOBAL_SHADER(FAutoExposureReduceFinalCS);
+		SHADER_USE_PARAMETER_STRUCT(FAutoExposureReduceFinalCS, FGlobalShader)
+
+	public:
+
+		BEGIN_SHADER_PARAMETER_STRUCT(FParameters, )
+			SHADER_PARAMETER(int32, InputSize)
+			SHADER_PARAMETER_RDG_BUFFER_UAV(RWBuffer, InputSums)
+			SHADER_PARAMETER_RDG_BUFFER_UAV(RWBuffer, InputCounts)
+			SHADER_PARAMETER_RDG_BUFFER_UAV(RWBuffer, OutputBuffer)
+		END_SHADER_PARAMETER_STRUCT()
+
+		static void ModifyCompilationEnvironment(const FGlobalShaderPermutationParameters& InParameters, FShaderCompilerEnvironment& OutEnvironment);
+	};
+
+} // namespace UE::NNEDenoiser::Private
\ No newline at end of file
diff --git a/Engine/Plugins/Experimental/NNEDenoiser/Source/NNEDenoiserShaders/Internal/NNEDenoiserShadersDefaultCS.h b/Engine/Plugins/Experimental/NNEDenoiser/Source/NNEDenoiserShaders/Internal/NNEDenoiserShadersDefaultCS.h
deleted file mode 100644
diff --git a/Engine/Plugins/Experimental/NNEDenoiser/Source/NNEDenoiserShaders/Internal/NNEDenoiserShadersDefaultIOProcessCS.h b/Engine/Plugins/Experimental/NNEDenoiser/Source/NNEDenoiserShaders/Internal/NNEDenoiserShadersDefaultIOProcessCS.h
new file mode 100644
--- /dev/null
+++ b/Engine/Plugins/Experimental/NNEDenoiser/Source/NNEDenoiserShaders/Internal/NNEDenoiserShadersDefaultIOProcessCS.h
@@ -0,0 +1,50 @@
+// Copyright Epic Games, Inc. All Rights Reserved.
+
+#pragma once
+
+#include "GlobalShader.h"
+#include "RenderGraphUtils.h"
+#include "ShaderParameterUtils.h"
+
+namespace UE::NNEDenoiserShaders::Internal
+{
+
+	enum class EDefaultIOProcessInputKind : uint8
+	{
+		Color = 0,
+		Albedo,
+		Normal,
+		Flow,
+		Output,
+		MAX
+	};
+
+	class FDefaultIOProcessConstants
+	{
+	public:
+		static constexpr int32 THREAD_GROUP_SIZE{16};
+	};
+	
+	class NNEDENOISERSHADERS_API FDefaultIOProcessCS : public FGlobalShader
+	{
+		DECLARE_GLOBAL_SHADER(FDefaultIOProcessCS);
+		SHADER_USE_PARAMETER_STRUCT(FDefaultIOProcessCS, FGlobalShader)
+
+		class FDefaultIOProcessInputKind : SHADER_PERMUTATION_ENUM_CLASS("INPUT_KIND_INDEX", EDefaultIOProcessInputKind);
+		using FPermutationDomain = TShaderPermutationDomain<FDefaultIOProcessInputKind>;
+
+	public:
+
+		BEGIN_SHADER_PARAMETER_STRUCT(FParameters, )
+			SHADER_PARAMETER(int32, InputTextureWidth)
+			SHADER_PARAMETER(int32, InputTextureHeight)
+			SHADER_PARAMETER_RDG_TEXTURE(Texture2D, InputTexture)
+			SHADER_PARAMETER(int32, OutputTextureWidth)
+			SHADER_PARAMETER(int32, OutputTextureHeight)
+			SHADER_PARAMETER_RDG_TEXTURE_UAV(RWTexture2D, OutputTexture)
+		END_SHADER_PARAMETER_STRUCT()
+
+		static void ModifyCompilationEnvironment(const FGlobalShaderPermutationParameters& InParameters, FShaderCompilerEnvironment& OutEnvironment);
+	};
+
+} // namespace UE::NNEDenoiser::Private
\ No newline at end of file
diff --git a/Engine/Plugins/Experimental/NNEDenoiser/Source/NNEDenoiserShaders/Internal/NNEDenoiserShadersMappedCopyCS.h b/Engine/Plugins/Experimental/NNEDenoiser/Source/NNEDenoiserShaders/Internal/NNEDenoiserShadersMappedCopyCS.h
new file mode 100644
--- /dev/null
+++ b/Engine/Plugins/Experimental/NNEDenoiser/Source/NNEDenoiserShaders/Internal/NNEDenoiserShadersMappedCopyCS.h
@@ -0,0 +1,117 @@
+// Copyright Epic Games, Inc. All Rights Reserved.
+
+#pragma once
+
+#include "GlobalShader.h"
+#include "RenderGraphUtils.h"
+#include "ShaderParameterUtils.h"
+
+namespace UE::NNEDenoiserShaders::Internal
+{
+
+	enum class ENNEDenoiserDataType : uint8
+	{
+		None = 0,
+		// Char,
+		// Boolean,
+		Half = 3,
+		Float,
+		MAX
+	};
+
+	class FMappedCopyConstants
+	{
+	public:
+		static constexpr int32 THREAD_GROUP_SIZE{16};
+		static constexpr int32 MAX_NUM_MAPPED_CHANNELS{4};
+	};
+	
+	class NNEDENOISERSHADERS_API FNNEDenoiserTextureBufferMappedCopyCS : public FGlobalShader
+	{
+		DECLARE_GLOBAL_SHADER(FNNEDenoiserTextureBufferMappedCopyCS);
+		SHADER_USE_PARAMETER_STRUCT(FNNEDenoiserTextureBufferMappedCopyCS, FGlobalShader)
+
+		class FNNEDenoiserDataType : SHADER_PERMUTATION_ENUM_CLASS("BUFFER_TYPE_INDEX", ENNEDenoiserDataType);
+		class FNNEDenoiserNumMappedChannels : SHADER_PERMUTATION_RANGE_INT("NUM_MAPPED_CHANNELS", 0, FMappedCopyConstants::MAX_NUM_MAPPED_CHANNELS + 1);
+		using FPermutationDomain = TShaderPermutationDomain<FNNEDenoiserDataType, FNNEDenoiserNumMappedChannels>;
+
+	public:
+
+		BEGIN_SHADER_PARAMETER_STRUCT(FParameters, )
+			SHADER_PARAMETER(int32, Width)
+			SHADER_PARAMETER(int32, Height)
+			SHADER_PARAMETER_RDG_TEXTURE(Texture2D, InputTexture)
+			SHADER_PARAMETER_RDG_BUFFER_UAV(RWBuffer, OutputBuffer)
+			SHADER_PARAMETER_ARRAY(FIntVector4, OutputChannel_InputChannel_Unused_Unused, [FMappedCopyConstants::MAX_NUM_MAPPED_CHANNELS])
+		END_SHADER_PARAMETER_STRUCT()
+
+		static void ModifyCompilationEnvironment(const FGlobalShaderPermutationParameters& InParameters, FShaderCompilerEnvironment& OutEnvironment);
+	};
+
+	class NNEDENOISERSHADERS_API FNNEDenoiserBufferMappedCopyCS : public FGlobalShader
+	{
+		DECLARE_GLOBAL_SHADER(FNNEDenoiserBufferMappedCopyCS);
+		SHADER_USE_PARAMETER_STRUCT(FNNEDenoiserBufferMappedCopyCS, FGlobalShader)
+
+		class FNNEDenoiserDataType : SHADER_PERMUTATION_ENUM_CLASS("BUFFER_TYPE_INDEX", ENNEDenoiserDataType);
+		class FNNEDenoiserNumMappedChannels : SHADER_PERMUTATION_RANGE_INT("NUM_MAPPED_CHANNELS", 0, FMappedCopyConstants::MAX_NUM_MAPPED_CHANNELS + 1);
+		using FPermutationDomain = TShaderPermutationDomain<FNNEDenoiserDataType, FNNEDenoiserNumMappedChannels>;
+
+	public:
+
+		BEGIN_SHADER_PARAMETER_STRUCT(FParameters, )
+			SHADER_PARAMETER(int32, Width)
+			SHADER_PARAMETER(int32, Height)
+			SHADER_PARAMETER_RDG_BUFFER_UAV(RWBuffer, InputBuffer)
+			SHADER_PARAMETER_RDG_BUFFER_UAV(RWBuffer, OutputBuffer)
+			SHADER_PARAMETER_ARRAY(FIntVector4, OutputChannel_InputChannel_Unused_Unused, [FMappedCopyConstants::MAX_NUM_MAPPED_CHANNELS])
+		END_SHADER_PARAMETER_STRUCT()
+
+		static void ModifyCompilationEnvironment(const FGlobalShaderPermutationParameters& InParameters, FShaderCompilerEnvironment& OutEnvironment);
+	};
+
+	class NNEDENOISERSHADERS_API FNNEDenoiserBufferTextureMappedCopyCS : public FGlobalShader
+	{
+		DECLARE_GLOBAL_SHADER(FNNEDenoiserBufferTextureMappedCopyCS);
+		SHADER_USE_PARAMETER_STRUCT(FNNEDenoiserBufferTextureMappedCopyCS, FGlobalShader)
+
+		class FNNEDenoiserDataType : SHADER_PERMUTATION_ENUM_CLASS("BUFFER_TYPE_INDEX", ENNEDenoiserDataType);
+		class FNNEDenoiserNumMappedChannels : SHADER_PERMUTATION_RANGE_INT("NUM_MAPPED_CHANNELS", 0, FMappedCopyConstants::MAX_NUM_MAPPED_CHANNELS + 1);
+		using FPermutationDomain = TShaderPermutationDomain<FNNEDenoiserDataType, FNNEDenoiserNumMappedChannels>;
+
+	public:
+
+		BEGIN_SHADER_PARAMETER_STRUCT(FParameters, )
+			SHADER_PARAMETER(int32, Width)
+			SHADER_PARAMETER(int32, Height)
+			SHADER_PARAMETER_RDG_BUFFER_UAV(RWBuffer, InputBuffer)
+			SHADER_PARAMETER_RDG_TEXTURE_UAV(RWTexture2D, OutputTexture)
+			SHADER_PARAMETER_ARRAY(FIntVector4, OutputChannel_InputChannel_Unused_Unused, [FMappedCopyConstants::MAX_NUM_MAPPED_CHANNELS])
+		END_SHADER_PARAMETER_STRUCT()
+
+		static void ModifyCompilationEnvironment(const FGlobalShaderPermutationParameters& InParameters, FShaderCompilerEnvironment& OutEnvironment);
+	};
+
+	class NNEDENOISERSHADERS_API FNNEDenoiserTextureMappedCopyCS : public FGlobalShader
+	{
+		DECLARE_GLOBAL_SHADER(FNNEDenoiserTextureMappedCopyCS);
+		SHADER_USE_PARAMETER_STRUCT(FNNEDenoiserTextureMappedCopyCS, FGlobalShader)
+
+		class FNNEDenoiserDataType : SHADER_PERMUTATION_ENUM_CLASS("BUFFER_TYPE_INDEX", ENNEDenoiserDataType);
+		class FNNEDenoiserNumMappedChannels : SHADER_PERMUTATION_RANGE_INT("NUM_MAPPED_CHANNELS", 0, FMappedCopyConstants::MAX_NUM_MAPPED_CHANNELS + 1);
+		using FPermutationDomain = TShaderPermutationDomain<FNNEDenoiserDataType, FNNEDenoiserNumMappedChannels>;
+
+	public:
+
+		BEGIN_SHADER_PARAMETER_STRUCT(FParameters, )
+			SHADER_PARAMETER(int32, Width)
+			SHADER_PARAMETER(int32, Height)
+			SHADER_PARAMETER_RDG_TEXTURE_UAV(RWTexture2D, InputTexture)
+			SHADER_PARAMETER_RDG_TEXTURE_UAV(RWTexture2D, OutputTexture)
+			SHADER_PARAMETER_ARRAY(FIntVector4, OutputChannel_InputChannel_Unused_Unused, [FMappedCopyConstants::MAX_NUM_MAPPED_CHANNELS])
+		END_SHADER_PARAMETER_STRUCT()
+
+		static void ModifyCompilationEnvironment(const FGlobalShaderPermutationParameters& InParameters, FShaderCompilerEnvironment& OutEnvironment);
+	};
+
+} // namespace UE::NNEDenoiser::Private
\ No newline at end of file
diff --git a/Engine/Plugins/Experimental/NNEDenoiser/Source/NNEDenoiserShaders/Internal/NNEDenoiserShadersOidnCS.h b/Engine/Plugins/Experimental/NNEDenoiser/Source/NNEDenoiserShaders/Internal/NNEDenoiserShadersOidnCS.h
deleted file mode 100644
diff --git a/Engine/Plugins/Experimental/NNEDenoiser/Source/NNEDenoiserShaders/Internal/NNEDenoiserShadersTransferFunctionOidnCS.h b/Engine/Plugins/Experimental/NNEDenoiser/Source/NNEDenoiserShaders/Internal/NNEDenoiserShadersTransferFunctionOidnCS.h
new file mode 100644
--- /dev/null
+++ b/Engine/Plugins/Experimental/NNEDenoiser/Source/NNEDenoiserShaders/Internal/NNEDenoiserShadersTransferFunctionOidnCS.h
@@ -0,0 +1,50 @@
+// Copyright Epic Games, Inc. All Rights Reserved.
+
+#pragma once
+
+#include "GlobalShader.h"
+#include "RenderGraphUtils.h"
+#include "ShaderParameterUtils.h"
+
+namespace UE::NNEDenoiserShaders::Internal
+{
+
+	enum class ETransferFunctionOidnMode : uint8
+	{
+		Forwward = 0,
+		Inverse,
+		MAX
+	};
+
+	class FTransferFunctionOidnConstants
+	{
+	public:
+		static constexpr int32 THREAD_GROUP_SIZE{16};
+	};
+	
+	class NNEDENOISERSHADERS_API FTransferFunctionOidnCS : public FGlobalShader
+	{
+		DECLARE_GLOBAL_SHADER(FTransferFunctionOidnCS);
+		SHADER_USE_PARAMETER_STRUCT(FTransferFunctionOidnCS, FGlobalShader)
+
+		class FTransferFunctionOidnMode : SHADER_PERMUTATION_ENUM_CLASS("MODE", ETransferFunctionOidnMode);
+		using FPermutationDomain = TShaderPermutationDomain<FTransferFunctionOidnMode>;
+
+	public:
+
+		BEGIN_SHADER_PARAMETER_STRUCT(FParameters, )
+			SHADER_PARAMETER(int32, InputTextureWidth)
+			SHADER_PARAMETER(int32, InputTextureHeight)
+			SHADER_PARAMETER_RDG_TEXTURE(Texture2D, InputTexture)
+			SHADER_PARAMETER(int32, OutputTextureWidth)
+			SHADER_PARAMETER(int32, OutputTextureHeight)
+			SHADER_PARAMETER_RDG_TEXTURE_UAV(RWTexture2D, OutputTexture)
+			SHADER_PARAMETER_RDG_BUFFER_UAV(RWBuffer, InputScaleBuffer)
+			SHADER_PARAMETER(float, NormScale)
+			SHADER_PARAMETER(float, InvNormScale)
+		END_SHADER_PARAMETER_STRUCT()
+
+		static void ModifyCompilationEnvironment(const FGlobalShaderPermutationParameters& InParameters, FShaderCompilerEnvironment& OutEnvironment);
+	};
+
+} // namespace UE::NNEDenoiser::Private
\ No newline at end of file
diff --git a/Engine/Plugins/Experimental/NNEDenoiser/Source/NNEDenoiserShaders/Private/NNEDenoiserShadersAutoExposureCS.cpp b/Engine/Plugins/Experimental/NNEDenoiser/Source/NNEDenoiserShaders/Private/NNEDenoiserShadersAutoExposureCS.cpp
new file mode 100644
--- /dev/null
+++ b/Engine/Plugins/Experimental/NNEDenoiser/Source/NNEDenoiserShaders/Private/NNEDenoiserShadersAutoExposureCS.cpp
@@ -0,0 +1,40 @@
+// Copyright Epic Games, Inc. All Rights Reserved.
+
+#include "NNEDenoiserShadersAutoExposureCS.h"
+
+namespace UE::NNEDenoiserShaders::Internal
+{
+	void FAutoExposureCS::ModifyCompilationEnvironment(const FGlobalShaderPermutationParameters& InParameters, FShaderCompilerEnvironment& OutEnvironment)
+	{
+		FGlobalShader::ModifyCompilationEnvironment(InParameters, OutEnvironment);
+		OutEnvironment.SetDefine(TEXT("MAX_FLT"), MAX_FLT);
+	}
+
+	void FAutoExposureDownsampleCS::ModifyCompilationEnvironment(const FGlobalShaderPermutationParameters& InParameters, FShaderCompilerEnvironment& OutEnvironment)
+	{
+		FGlobalShader::ModifyCompilationEnvironment(InParameters, OutEnvironment);
+		OutEnvironment.SetDefine(TEXT("THREAD_GROUP_SIZE"), FAutoExposureDownsampleConstants::THREAD_GROUP_SIZE);
+		OutEnvironment.SetDefine(TEXT("MAX_FLT"), MAX_FLT);
+	}
+
+	void FAutoExposureReduceCS::ModifyCompilationEnvironment(const FGlobalShaderPermutationParameters& InParameters, FShaderCompilerEnvironment& OutEnvironment)
+	{
+		FGlobalShader::ModifyCompilationEnvironment(InParameters, OutEnvironment);
+		OutEnvironment.SetDefine(TEXT("THREAD_GROUP_SIZE"), FAutoExposureReduceConstants::THREAD_GROUP_SIZE);
+		OutEnvironment.SetDefine(TEXT("EPS"), FAutoExposureReduceConstants::EPS);
+	}
+
+	void FAutoExposureReduceFinalCS::ModifyCompilationEnvironment(const FGlobalShaderPermutationParameters& InParameters, FShaderCompilerEnvironment& OutEnvironment)
+	{
+		FGlobalShader::ModifyCompilationEnvironment(InParameters, OutEnvironment);
+		OutEnvironment.SetDefine(TEXT("THREAD_GROUP_SIZE"), FAutoExposureReduceConstants::THREAD_GROUP_SIZE);
+		OutEnvironment.SetDefine(TEXT("KEY"), FAutoExposureReduceConstants::KEY);
+	}
+
+	IMPLEMENT_GLOBAL_SHADER(FAutoExposureCS, "/NNEDenoiserShaders/NNEDenoiserShadersAutoExposure.usf", "AutoExposure", SF_Compute);
+
+	IMPLEMENT_GLOBAL_SHADER(FAutoExposureDownsampleCS, "/NNEDenoiserShaders/NNEDenoiserShadersAutoExposureDownsample.usf", "Downsample", SF_Compute);
+	IMPLEMENT_GLOBAL_SHADER(FAutoExposureReduceCS, "/NNEDenoiserShaders/NNEDenoiserShadersAutoExposureReduce.usf", "Reduce", SF_Compute);
+	IMPLEMENT_GLOBAL_SHADER(FAutoExposureReduceFinalCS, "/NNEDenoiserShaders/NNEDenoiserShadersAutoExposureReduceFinal.usf", "ReduceFinal", SF_Compute);
+
+} // UE::NNEDenoiser::Private
\ No newline at end of file
diff --git a/Engine/Plugins/Experimental/NNEDenoiser/Source/NNEDenoiserShaders/Private/NNEDenoiserShadersDefaultCS.cpp b/Engine/Plugins/Experimental/NNEDenoiser/Source/NNEDenoiserShaders/Private/NNEDenoiserShadersDefaultCS.cpp
deleted file mode 100644
diff --git a/Engine/Plugins/Experimental/NNEDenoiser/Source/NNEDenoiserShaders/Private/NNEDenoiserShadersDefaultIOProcessCS.cpp b/Engine/Plugins/Experimental/NNEDenoiser/Source/NNEDenoiserShaders/Private/NNEDenoiserShadersDefaultIOProcessCS.cpp
new file mode 100644
--- /dev/null
+++ b/Engine/Plugins/Experimental/NNEDenoiser/Source/NNEDenoiserShaders/Private/NNEDenoiserShadersDefaultIOProcessCS.cpp
@@ -0,0 +1,16 @@
+// Copyright Epic Games, Inc. All Rights Reserved.
+
+#include "NNEDenoiserShadersDefaultIOProcessCS.h"
+
+namespace UE::NNEDenoiserShaders::Internal
+{
+	void FDefaultIOProcessCS::ModifyCompilationEnvironment(const FGlobalShaderPermutationParameters& InParameters, FShaderCompilerEnvironment& OutEnvironment)
+	{
+		FGlobalShader::ModifyCompilationEnvironment(InParameters, OutEnvironment);
+		OutEnvironment.SetDefine(TEXT("THREAD_GROUP_SIZE"), FDefaultIOProcessConstants::THREAD_GROUP_SIZE);
+		OutEnvironment.SetDefine(TEXT("MAX_FLT"), MAX_FLT);
+	}
+
+	IMPLEMENT_GLOBAL_SHADER(FDefaultIOProcessCS, "/NNEDenoiserShaders/NNEDenoiserShadersDefaultIOProcess.usf", "IOProcess", SF_Compute);
+
+} // UE::NNEDenoiser::Private
\ No newline at end of file
diff --git a/Engine/Plugins/Experimental/NNEDenoiser/Source/NNEDenoiserShaders/Private/NNEDenoiserShadersMappedCopyCS.cpp b/Engine/Plugins/Experimental/NNEDenoiser/Source/NNEDenoiserShaders/Private/NNEDenoiserShadersMappedCopyCS.cpp
new file mode 100644
--- /dev/null
+++ b/Engine/Plugins/Experimental/NNEDenoiser/Source/NNEDenoiserShaders/Private/NNEDenoiserShadersMappedCopyCS.cpp
@@ -0,0 +1,48 @@
+// Copyright Epic Games, Inc. All Rights Reserved.
+
+#include "NNEDenoiserShadersMappedCopyCS.h"
+
+namespace UE::NNEDenoiserShaders::Internal
+{
+	void CommonModifyCompilationEnvironment(const FGlobalShaderPermutationParameters& InParameters, FShaderCompilerEnvironment& OutEnvironment)
+	{
+		FGlobalShader::ModifyCompilationEnvironment(InParameters, OutEnvironment);
+		OutEnvironment.SetDefine(TEXT("THREAD_GROUP_SIZE"), FMappedCopyConstants::THREAD_GROUP_SIZE);
+		OutEnvironment.SetDefine(TEXT("MAX_NUM_MAPPED_CHANNELS"), FMappedCopyConstants::MAX_NUM_MAPPED_CHANNELS);
+	}
+
+	void FNNEDenoiserTextureBufferMappedCopyCS::ModifyCompilationEnvironment(const FGlobalShaderPermutationParameters& InParameters, FShaderCompilerEnvironment& OutEnvironment)
+	{
+		CommonModifyCompilationEnvironment(InParameters, OutEnvironment);
+		OutEnvironment.SetDefine(TEXT("INPUT_TYPE"), 0);
+		OutEnvironment.SetDefine(TEXT("OUTPUT_TYPE"), 1);
+	}
+
+	void FNNEDenoiserBufferMappedCopyCS::ModifyCompilationEnvironment(const FGlobalShaderPermutationParameters& InParameters, FShaderCompilerEnvironment& OutEnvironment)
+	{
+		CommonModifyCompilationEnvironment(InParameters, OutEnvironment);
+		OutEnvironment.SetDefine(TEXT("INPUT_TYPE"), 1);
+		OutEnvironment.SetDefine(TEXT("OUTPUT_TYPE"), 1);
+	}
+
+	void FNNEDenoiserBufferTextureMappedCopyCS::ModifyCompilationEnvironment(const FGlobalShaderPermutationParameters& InParameters, FShaderCompilerEnvironment& OutEnvironment)
+	{
+		CommonModifyCompilationEnvironment(InParameters, OutEnvironment);
+		OutEnvironment.SetDefine(TEXT("INPUT_TYPE"), 1);
+		OutEnvironment.SetDefine(TEXT("OUTPUT_TYPE"), 0);
+	}
+
+	void FNNEDenoiserTextureMappedCopyCS::ModifyCompilationEnvironment(const FGlobalShaderPermutationParameters& InParameters, FShaderCompilerEnvironment& OutEnvironment)
+	{
+		CommonModifyCompilationEnvironment(InParameters, OutEnvironment);
+		OutEnvironment.SetDefine(TEXT("INPUT_TYPE"), 0);
+		OutEnvironment.SetDefine(TEXT("OUTPUT_TYPE"), 0);
+	}
+
+	IMPLEMENT_GLOBAL_SHADER(FNNEDenoiserTextureBufferMappedCopyCS, "/NNEDenoiserShaders/NNEDenoiserShadersMappedCopy.usf", "MappedCopy", SF_Compute);
+	IMPLEMENT_GLOBAL_SHADER(FNNEDenoiserBufferMappedCopyCS, "/NNEDenoiserShaders/NNEDenoiserShadersMappedCopy.usf", "MappedCopy", SF_Compute);
+
+	IMPLEMENT_GLOBAL_SHADER(FNNEDenoiserBufferTextureMappedCopyCS, "/NNEDenoiserShaders/NNEDenoiserShadersMappedCopy.usf", "MappedCopy", SF_Compute);
+	IMPLEMENT_GLOBAL_SHADER(FNNEDenoiserTextureMappedCopyCS, "/NNEDenoiserShaders/NNEDenoiserShadersMappedCopy.usf", "MappedCopy", SF_Compute);
+
+} // UE::NNEDenoiser::Private
\ No newline at end of file
diff --git a/Engine/Plugins/Experimental/NNEDenoiser/Source/NNEDenoiserShaders/Private/NNEDenoiserShadersOidn.cpp b/Engine/Plugins/Experimental/NNEDenoiser/Source/NNEDenoiserShaders/Private/NNEDenoiserShadersOidn.cpp
deleted file mode 100644
diff --git a/Engine/Plugins/Experimental/NNEDenoiser/Source/NNEDenoiserShaders/Private/NNEDenoiserTransferFunctionOidn.cpp b/Engine/Plugins/Experimental/NNEDenoiser/Source/NNEDenoiserShaders/Private/NNEDenoiserTransferFunctionOidn.cpp
new file mode 100644
--- /dev/null
+++ b/Engine/Plugins/Experimental/NNEDenoiser/Source/NNEDenoiserShaders/Private/NNEDenoiserTransferFunctionOidn.cpp
@@ -0,0 +1,16 @@
+// Copyright Epic Games, Inc. All Rights Reserved.
+
+#include "NNEDenoiserShadersTransferFunctionOidnCS.h"
+
+namespace UE::NNEDenoiserShaders::Internal
+{
+	void FTransferFunctionOidnCS::ModifyCompilationEnvironment(const FGlobalShaderPermutationParameters& InParameters, FShaderCompilerEnvironment& OutEnvironment)
+	{
+		FGlobalShader::ModifyCompilationEnvironment(InParameters, OutEnvironment);
+		OutEnvironment.SetDefine(TEXT("THREAD_GROUP_SIZE"), FTransferFunctionOidnConstants::THREAD_GROUP_SIZE);
+		OutEnvironment.SetDefine(TEXT("MAX_FLT"), MAX_FLT);
+	}
+
+	IMPLEMENT_GLOBAL_SHADER(FTransferFunctionOidnCS, "/NNEDenoiserShaders/NNEDenoiserShadersTransferFunctionOidn.usf", "PreOrPostprocess", SF_Compute);
+
+} // UE::NNEDenoiser::Private
\ No newline at end of file
diff --git a/Engine/Plugins/Experimental/NNEDenoiser/Source/NNEDenoiserShaders/Shaders/NNEDenoiserShadersAutoExposure.usf b/Engine/Plugins/Experimental/NNEDenoiser/Source/NNEDenoiserShaders/Shaders/NNEDenoiserShadersAutoExposure.usf
new file mode 100644
--- /dev/null
+++ b/Engine/Plugins/Experimental/NNEDenoiser/Source/NNEDenoiserShaders/Shaders/NNEDenoiserShadersAutoExposure.usf
@@ -0,0 +1,40 @@
+// Copyright Epic Games, Inc. All Rights Reserved.
+
+#include "/Engine/Public/Platform.ush"
+#include "/Engine/Private/Common.ush"
+
+int InputTextureWidth;
+int InputTextureHeight;
+Texture2D InputTexture;
+RWBuffer<float> OutputBuffer;
+
+[numthreads(1, 1, 1)]
+void AutoExposure(in const uint3 DispatchThreadID : SV_DispatchThreadID)
+{
+	const float Key = 0.18;
+
+	float Sum = 0.0;
+	int Count = 0;
+	for (int i = 0; i < InputTextureWidth; i++)
+	{
+		for (int j = 0; j < InputTextureHeight; j++)
+		{
+			float3 Color = InputTexture.Load(int3(i, j, 0)).rgb;
+
+			Color = clamp(Color, float3(0.0, 0.0, 0.0), float3(MAX_FLT, MAX_FLT, MAX_FLT));
+
+			float Lum = Luminance(Color);
+			if (Lum > 1e-8)
+			{
+				Sum += log2(Lum);
+				Count++;
+			}
+		}
+	}
+
+	float Result = Count > 0 ? Sum / (float)Count : 1.0;
+	Result = Key / exp2(Result);
+
+	OutputBuffer[0] = Result;
+	OutputBuffer[1] = 1.0 / Result;
+}
\ No newline at end of file
diff --git a/Engine/Plugins/Experimental/NNEDenoiser/Source/NNEDenoiserShaders/Shaders/NNEDenoiserShadersAutoExposureDownsample.usf b/Engine/Plugins/Experimental/NNEDenoiser/Source/NNEDenoiserShaders/Shaders/NNEDenoiserShadersAutoExposureDownsample.usf
new file mode 100644
--- /dev/null
+++ b/Engine/Plugins/Experimental/NNEDenoiser/Source/NNEDenoiserShaders/Shaders/NNEDenoiserShadersAutoExposureDownsample.usf
@@ -0,0 +1,55 @@
+// Copyright Epic Games, Inc. All Rights Reserved.
+
+#include "/Engine/Public/Platform.ush"
+#include "/Engine/Private/Common.ush"
+
+int InputTextureWidth;
+int InputTextureHeight;
+Texture2D InputTexture;
+
+int NumBinsW;
+int NumBinsH;
+RWBuffer<float> OutputBins;
+
+groupshared float Sums[THREAD_GROUP_SIZE * THREAD_GROUP_SIZE];
+
+[numthreads(THREAD_GROUP_SIZE, THREAD_GROUP_SIZE, 1)]
+void Downsample(in const uint3 GroupThreadID : SV_GroupThreadID, in const uint3 GroupID : SV_GroupID, in const uint3 DispatchThreadID : SV_DispatchThreadID)
+{
+	const int BeginW = GroupID.x * InputTextureWidth / NumBinsW;
+	const int BeginH = GroupID.y * InputTextureHeight / NumBinsH;
+	const int EndW = (GroupID.x + 1) * InputTextureWidth / NumBinsW;
+	const int EndH = (GroupID.y + 1) * InputTextureHeight / NumBinsH;
+
+	const int CoordW = BeginW + GroupThreadID.x;
+	const int CoordH = BeginH + GroupThreadID.y;
+
+	float Lum = 0.0;
+	if (CoordW < EndW && CoordH < EndH)
+	{
+		float3 Color = InputTexture.Load(int3(CoordW, CoordH, 0)).rgb;
+		Color = clamp(Color, float3(0.0, 0.0, 0.0), float3(MAX_FLT, MAX_FLT, MAX_FLT));
+
+		Lum = Luminance(Color);
+	}
+
+	const int LocalID = GroupThreadID.y * THREAD_GROUP_SIZE + GroupThreadID.x;
+	Sums[LocalID] = Lum;
+
+	for (int i = THREAD_GROUP_SIZE * THREAD_GROUP_SIZE / 2; i > 0; i >>= 1)
+	{
+		GroupMemoryBarrierWithGroupSync();
+		
+		if (LocalID < i)
+		{
+			Sums[LocalID] += Sums[LocalID + i];
+		}
+	}
+
+	if (LocalID == 0)
+	{
+		const float AverageLum = Sums[0] / float((EndW - BeginW) * (EndH - BeginH));
+		const int GroupLinearID = GroupID.y * NumBinsW + GroupID.x;
+		OutputBins[GroupLinearID] = AverageLum;
+	}
+}
\ No newline at end of file
diff --git a/Engine/Plugins/Experimental/NNEDenoiser/Source/NNEDenoiserShaders/Shaders/NNEDenoiserShadersAutoExposureReduce.usf b/Engine/Plugins/Experimental/NNEDenoiser/Source/NNEDenoiserShaders/Shaders/NNEDenoiserShadersAutoExposureReduce.usf
new file mode 100644
--- /dev/null
+++ b/Engine/Plugins/Experimental/NNEDenoiser/Source/NNEDenoiserShaders/Shaders/NNEDenoiserShadersAutoExposureReduce.usf
@@ -0,0 +1,50 @@
+// Copyright Epic Games, Inc. All Rights Reserved.
+
+#include "/Engine/Public/Platform.ush"
+#include "/Engine/Private/Common.ush"
+
+int InputSize;
+RWBuffer<float> InputBins;
+RWBuffer<float> OutputSums;
+RWBuffer<int> OutputCounts;
+int NumThreads;
+
+groupshared float Sums[THREAD_GROUP_SIZE];
+groupshared int Counts[THREAD_GROUP_SIZE];
+
+[numthreads(THREAD_GROUP_SIZE, 1, 1)]
+void Reduce(in const uint3 GroupThreadID : SV_GroupThreadID, in const uint3 GroupID : SV_GroupID, in const uint3 DispatchThreadID : SV_DispatchThreadID)
+{
+	float LocalSum = 0.0;
+	int LocalCount = 0;
+	for (int i = DispatchThreadID.x; i < InputSize; i += NumThreads)
+	{
+		const float Lum = InputBins[i];
+		if (Lum > EPS)
+		{
+			LocalSum += log2(Lum);
+			++LocalCount;
+		}
+	}
+
+	const int LocalID = GroupThreadID.x;
+	Sums[LocalID] = LocalSum;
+	Counts[LocalID] = LocalCount;
+
+	for (int i = THREAD_GROUP_SIZE / 2; i > 0; i >>= 1)
+	{
+		GroupMemoryBarrierWithGroupSync();
+
+		if (LocalID < i)
+		{
+			Sums[LocalID] += Sums[LocalID + i];
+			Counts[LocalID] += Counts[LocalID + i];
+		}
+	}
+
+	if (LocalID == 0)
+	{
+		OutputSums[GroupID.x] = Sums[0];
+		OutputCounts[GroupID.x] = Counts[0];
+	}
+}
\ No newline at end of file
diff --git a/Engine/Plugins/Experimental/NNEDenoiser/Source/NNEDenoiserShaders/Shaders/NNEDenoiserShadersAutoExposureReduceFinal.usf b/Engine/Plugins/Experimental/NNEDenoiser/Source/NNEDenoiserShaders/Shaders/NNEDenoiserShadersAutoExposureReduceFinal.usf
new file mode 100644
--- /dev/null
+++ b/Engine/Plugins/Experimental/NNEDenoiser/Source/NNEDenoiserShaders/Shaders/NNEDenoiserShadersAutoExposureReduceFinal.usf
@@ -0,0 +1,46 @@
+// Copyright Epic Games, Inc. All Rights Reserved.
+
+#include "/Engine/Public/Platform.ush"
+#include "/Engine/Private/Common.ush"
+
+int InputSize;
+RWBuffer<float> InputSums;
+RWBuffer<int> InputCounts;
+RWBuffer<float> OutputBuffer; // 0: Result, 1: 1.0 / Result
+
+groupshared float Sums[THREAD_GROUP_SIZE];
+groupshared int Counts[THREAD_GROUP_SIZE];
+
+[numthreads(THREAD_GROUP_SIZE, 1, 1)]
+void ReduceFinal(in const uint3 GroupThreadID : SV_GroupThreadID)
+{
+	const int LocalID = GroupThreadID.x;
+
+	if (LocalID < InputSize)
+	{
+		Sums[LocalID] = InputSums[LocalID];
+		Counts[LocalID] = InputCounts[LocalID];
+	}
+	else
+	{
+		Sums[LocalID] = 0;
+		Counts[LocalID] = 0;
+	}
+
+	for (int i = THREAD_GROUP_SIZE / 2; i > 0; i >>= 1)
+	{
+		GroupMemoryBarrierWithGroupSync();
+
+		if (LocalID < i)
+		{
+			Sums[LocalID] += Sums[LocalID + i];
+			Counts[LocalID] += Counts[LocalID + i];
+		}
+	}
+
+	if (LocalID == 0)
+	{
+		OutputBuffer[0] = (Counts[0] > 0) ? (KEY / exp2(Sums[0] / float(Counts[0]))) : 1.0;
+		OutputBuffer[1] = 1.0 / OutputBuffer[0];
+	}
+}
\ No newline at end of file
diff --git a/Engine/Plugins/Experimental/NNEDenoiser/Source/NNEDenoiserShaders/Shaders/NNEDenoiserShadersDefault.usf b/Engine/Plugins/Experimental/NNEDenoiser/Source/NNEDenoiserShaders/Shaders/NNEDenoiserShadersDefault.usf
deleted file mode 100644
diff --git a/Engine/Plugins/Experimental/NNEDenoiser/Source/NNEDenoiserShaders/Shaders/NNEDenoiserShadersDefaultIOProcess.usf b/Engine/Plugins/Experimental/NNEDenoiser/Source/NNEDenoiserShaders/Shaders/NNEDenoiserShadersDefaultIOProcess.usf
new file mode 100644
--- /dev/null
+++ b/Engine/Plugins/Experimental/NNEDenoiser/Source/NNEDenoiserShaders/Shaders/NNEDenoiserShadersDefaultIOProcess.usf
@@ -0,0 +1,44 @@
+// Copyright Epic Games, Inc. All Rights Reserved.
+
+#include "/Engine/Public/Platform.ush"
+#include "/Engine/Private/Common.ush"
+
+int InputTextureWidth;
+int InputTextureHeight;
+Texture2D InputTexture;
+
+int OutputTextureWidth;
+int OutputTextureHeight;
+RWTexture2D<float4> OutputTexture;
+
+float3 GetProcessedValue(float3 Value)
+{
+#if INPUT_KIND_INDEX == 0 // Input color
+	return clamp(Value, float3(0.0, 0.0, 0.0), float3(MAX_FLT, MAX_FLT, MAX_FLT));
+#elif INPUT_KIND_INDEX == 1 // Input albedo
+	return clamp(Value, float3(0.0, 0.0, 0.0), float3(1.0, 1.0, 1.0));
+#elif INPUT_KIND_INDEX == 2 // Input normal
+	return clamp(Value, float3(-1.0, -1.0, -1.0), float3(1.0, 1.0, 1.0)) * 0.5 + float3(0.5, 0.5, 0.5);
+#elif INPUT_KIND_INDEX == 3 // Input flow
+	// TODO flow texture check for any necessary conversion
+	return Value;
+#else
+	// INPUT_KIND_INDEX == 4 // Output color
+	return clamp(Value, float3(0.0, 0.0, 0.0), float3(MAX_FLT, MAX_FLT, MAX_FLT));
+#endif
+}
+
+
+[numthreads(THREAD_GROUP_SIZE, THREAD_GROUP_SIZE, 1)]
+void IOProcess(in const uint3 DispatchThreadID : SV_DispatchThreadID)
+{
+	if (DispatchThreadID.x >= OutputTextureWidth || DispatchThreadID.y >= OutputTextureHeight)
+	{
+		return;
+	}
+
+	const float4 Color = InputTexture.Load(int3(DispatchThreadID.x, DispatchThreadID.y, 0));
+
+	OutputTexture[DispatchThreadID.xy].rgb = GetProcessedValue(Color.rgb);
+	OutputTexture[DispatchThreadID.xy].a = Color.a;
+}
\ No newline at end of file
diff --git a/Engine/Plugins/Experimental/NNEDenoiser/Source/NNEDenoiserShaders/Shaders/NNEDenoiserShadersMappedCopy.usf b/Engine/Plugins/Experimental/NNEDenoiser/Source/NNEDenoiserShaders/Shaders/NNEDenoiserShadersMappedCopy.usf
new file mode 100644
--- /dev/null
+++ b/Engine/Plugins/Experimental/NNEDenoiser/Source/NNEDenoiserShaders/Shaders/NNEDenoiserShadersMappedCopy.usf
@@ -0,0 +1,69 @@
+// Copyright Epic Games, Inc. All Rights Reserved.
+
+#include "/Engine/Public/Platform.ush"
+#include "/Engine/Private/Common.ush"
+
+#if BUFFER_TYPE_INDEX == 4 // Float16
+	#define BUFFER_TYPE half
+	#define READ(X) float(X)
+	#define WRITE(X) half(X)
+#else
+	// BUFFER_TYPE_INDEX == 3 // Float32
+	#define BUFFER_TYPE float
+	#define READ(X) X
+	#define WRITE(X) X
+#endif
+
+int Width;
+int Height;
+
+int4 OutputChannel_InputChannel_Unused_Unused[MAX_NUM_MAPPED_CHANNELS];
+
+#if INPUT_TYPE == 0 // Texture
+Texture2D InputTexture;
+#else // INPUT_TYPE == 1 // Buffer
+RWBuffer<BUFFER_TYPE> InputBuffer;
+#endif
+
+#if OUTPUT_TYPE == 0 // Texture
+RWTexture2D<float4> OutputTexture;
+#else // OUTPUT_TYPE == 1 // Buffer
+RWBuffer<BUFFER_TYPE> OutputBuffer;
+#endif
+
+[numthreads(THREAD_GROUP_SIZE, THREAD_GROUP_SIZE, 1)]
+void MappedCopy(in const uint3 DispatchThreadID : SV_DispatchThreadID)
+{
+	if (DispatchThreadID.x >= Width || DispatchThreadID.y >= Height)
+	{
+		return;
+	}
+
+#if INPUT_TYPE == 1 || OUTPUT_TYPE == 1
+	const int Idx = DispatchThreadID.y * Width + DispatchThreadID.x;
+	const int Offset = Width * Height;
+#endif
+
+#if INPUT_TYPE == 0
+	const float4 InputColor = InputTexture.Load(int3(DispatchThreadID.x, DispatchThreadID.y, 0));
+#endif
+
+	for (int I = 0; I < NUM_MAPPED_CHANNELS; I++)
+	{
+		const int4 Info = OutputChannel_InputChannel_Unused_Unused[I];
+		const int OutputChannel = Info.x;
+		const int InputChannel = Info.y;
+
+#if INPUT_TYPE == 0
+		const float InputValue = InputColor[InputChannel];
+#else
+		const float InputValue = InputBuffer[Idx + InputChannel * Offset];
+#endif
+
+#if OUTPUT_TYPE == 0
+		OutputTexture[DispatchThreadID.xy][OutputChannel] = READ(InputValue);
+#else
+		OutputBuffer[Idx + OutputChannel * Offset] = WRITE(InputValue);
+#endif
+	}
+}
\ No newline at end of file
diff --git a/Engine/Plugins/Experimental/NNEDenoiser/Source/NNEDenoiserShaders/Shaders/NNEDenoiserShadersOidn.usf b/Engine/Plugins/Experimental/NNEDenoiser/Source/NNEDenoiserShaders/Shaders/NNEDenoiserShadersOidn.usf
deleted file mode 100644
diff --git a/Engine/Plugins/Experimental/NNEDenoiser/Source/NNEDenoiserShaders/Shaders/NNEDenoiserShadersTransferFunctionOidn.usf b/Engine/Plugins/Experimental/NNEDenoiser/Source/NNEDenoiserShaders/Shaders/NNEDenoiserShadersTransferFunctionOidn.usf
new file mode 100644
--- /dev/null
+++ b/Engine/Plugins/Experimental/NNEDenoiser/Source/NNEDenoiserShaders/Shaders/NNEDenoiserShadersTransferFunctionOidn.usf
@@ -0,0 +1,81 @@
+// Copyright Epic Games, Inc. All Rights Reserved.
+
+#include "/Engine/Public/Platform.ush"
+#include "/Engine/Private/Common.ush"
+
+int InputTextureWidth;
+int InputTextureHeight;
+Texture2D InputTexture;
+
+int OutputTextureWidth;
+int OutputTextureHeight;
+RWTexture2D<float4> OutputTexture;
+
+RWBuffer<float> InputScaleBuffer;
+
+float NormScale;
+float InvNormScale;
+
+///
+
+// from [REDACTED_PROFILE_URL]/oidn/blob/master/core/color.h
+// struct FPU
+// {
+static float a  =  1.41283765e+03;
+static float b  =  1.64593172e+00;
+static float c  =  4.31384981e-01;
+static float d  = -2.94139609e-03;
+static float e  =  1.92653254e-01;
+static float f  =  6.26026094e-03;
+static float g  =  9.98620152e-01;
+static float y0 =  1.57945760e-06;
+static float y1 =  3.22087631e-02;
+static float x0 =  2.23151711e-03;
+static float x1 =  3.70974749e-01;
+
+static float Forward(float y)
+{
+	if (y <= y0) return a * y;
+	else if (y <= y1) return b * pow(y, c) + d;
+	return e * log(y + f) + g;
+}
+
+static float Inverse(float x)
+{
+	if (x <= x0) return x / a;
+	else if (x <= x1) return pow((x - d) / b, 1.0 / c);
+	return exp((x - g) / e) - f;
+}
+
+float3 GetProcessedValue(float3 Value)
+{
+#if MODE == 0 // Forward
+	Value = Value * InputScaleBuffer[0];
+	Value = clamp(Value, float3(0.0, 0.0, 0.0), float3(MAX_FLT, MAX_FLT, MAX_FLT));
+
+	return float3(Forward(Value.r), Forward(Value.g), Forward(Value.b)) * NormScale;
+#else
+	// MODE == 1 // Inverse
+	Value = clamp(Value, float3(0.0, 0.0, 0.0), float3(MAX_FLT, MAX_FLT, MAX_FLT)) * InvNormScale;
+	
+	return float3(Inverse(Value.r), Inverse(Value.g), Inverse(Value.b)) * InputScaleBuffer[1];
+#endif
+}
+
+// };
+
+///
+
+[numthreads(THREAD_GROUP_SIZE, THREAD_GROUP_SIZE, 1)]
+void PreOrPostprocess(in const uint3 DispatchThreadID : SV_DispatchThreadID)
+{
+	if (DispatchThreadID.x >= OutputTextureWidth || DispatchThreadID.y >= OutputTextureHeight)
+	{
+		return;
+	}
+
+	const float4 Color = InputTexture.Load(int3(DispatchThreadID.x, DispatchThreadID.y, 0));
+
+	OutputTexture[DispatchThreadID.xy].rgb = GetProcessedValue(Color.rgb);
+	OutputTexture[DispatchThreadID.xy].a = Color.a;
+}
\ No newline at end of file
