Fixes for EAutomationTestFlags.

diff --git a/Engine/Plugins/Experimental/NNEDenoiser/Source/NNEDenoiser/Private/NNEDenoiserUnitTestOidn.cpp b/Engine/Plugins/Experimental/NNEDenoiser/Source/NNEDenoiser/Private/NNEDenoiserUnitTestOidn.cpp
--- a/Engine/Plugins/Experimental/NNEDenoiser/Source/NNEDenoiser/Private/NNEDenoiserUnitTestOidn.cpp
+++ b/Engine/Plugins/Experimental/NNEDenoiser/Source/NNEDenoiser/Private/NNEDenoiserUnitTestOidn.cpp
@@ -28,7 +28,7 @@ BEGIN_SHADER_PARAMETER_STRUCT(FTransferFunctionTestDownloadParameters, )
 	RDG_TEXTURE_ACCESS(OutputTexture, ERHIAccess::CopySrc)
 END_SHADER_PARAMETER_STRUCT()
 
-IMPLEMENT_SIMPLE_AUTOMATION_TEST(FTransferFunctionTest, "NNEDenoiser.UnitTests.OIDN.TransferFunction.Forward", EAutomationTestFlags::ApplicationContextMask | EAutomationTestFlags::FeatureMask | EAutomationTestFlags::EngineFilter)
+IMPLEMENT_SIMPLE_AUTOMATION_TEST(FTransferFunctionTest, "NNEDenoiser.UnitTests.OIDN.TransferFunction.Forward", EAutomationTestFlags::ApplicationContextMask | EAutomationTestFlags_FeatureMask | EAutomationTestFlags::EngineFilter)
 bool FTransferFunctionTest::RunTest(const FString& Parameter)
 {
 	const int32 Width = 1920;
diff --git a/Engine/Plugins/Experimental/NNEDenoiser/Source/NNEDenoiser/Private/NNEDenoiserUnitTests.cpp b/Engine/Plugins/Experimental/NNEDenoiser/Source/NNEDenoiser/Private/NNEDenoiserUnitTests.cpp
--- a/Engine/Plugins/Experimental/NNEDenoiser/Source/NNEDenoiser/Private/NNEDenoiserUnitTests.cpp
+++ b/Engine/Plugins/Experimental/NNEDenoiser/Source/NNEDenoiser/Private/NNEDenoiserUnitTests.cpp
@@ -24,7 +24,7 @@ BEGIN_SHADER_PARAMETER_STRUCT(FAutoExposureTestDownloadParameters, )
 	RDG_BUFFER_ACCESS(OutputBuffer, ERHIAccess::CopySrc)
 END_SHADER_PARAMETER_STRUCT()
 
-IMPLEMENT_SIMPLE_AUTOMATION_TEST(FAutoExposureTest, "NNEDenoiser.UnitTests.AutoExposure", EAutomationTestFlags::ApplicationContextMask | EAutomationTestFlags::FeatureMask | EAutomationTestFlags::EngineFilter)
+IMPLEMENT_SIMPLE_AUTOMATION_TEST(FAutoExposureTest, "NNEDenoiser.UnitTests.AutoExposure", EAutomationTestFlags::ApplicationContextMask | EAutomationTestFlags_FeatureMask | EAutomationTestFlags::EngineFilter)
 bool FAutoExposureTest::RunTest(const FString& Parameter)
 {
 	const int32 Width = 1920;
