NNEDenoiser: Fix minor shader warnings and avoid building on platforms that do not support raytracing

#rnx

diff --git a/Engine/Plugins/Experimental/NNEDenoiser/Source/NNEDenoiserShaders/Internal/NNEDenoiserShadersAutoExposureCS.h b/Engine/Plugins/Experimental/NNEDenoiser/Source/NNEDenoiserShaders/Internal/NNEDenoiserShadersAutoExposureCS.h
--- a/Engine/Plugins/Experimental/NNEDenoiser/Source/NNEDenoiserShaders/Internal/NNEDenoiserShadersAutoExposureCS.h
+++ b/Engine/Plugins/Experimental/NNEDenoiser/Source/NNEDenoiserShaders/Internal/NNEDenoiserShadersAutoExposureCS.h
@@ -24,6 +24,7 @@ namespace UE::NNEDenoiserShaders::Internal
 		END_SHADER_PARAMETER_STRUCT()
 
 		static void ModifyCompilationEnvironment(const FGlobalShaderPermutationParameters& InParameters, FShaderCompilerEnvironment& OutEnvironment);
+		static bool ShouldCompilePermutation(const FGlobalShaderPermutationParameters& Parameters);
 	};
 
 	class FAutoExposureDownsampleConstants
@@ -50,6 +51,7 @@ namespace UE::NNEDenoiserShaders::Internal
 		END_SHADER_PARAMETER_STRUCT()
 
 		static void ModifyCompilationEnvironment(const FGlobalShaderPermutationParameters& InParameters, FShaderCompilerEnvironment& OutEnvironment);
+		static bool ShouldCompilePermutation(const FGlobalShaderPermutationParameters& Parameters);
 	};
 
 	class FAutoExposureReduceConstants
@@ -76,6 +78,7 @@ namespace UE::NNEDenoiserShaders::Internal
 		END_SHADER_PARAMETER_STRUCT()
 
 		static void ModifyCompilationEnvironment(const FGlobalShaderPermutationParameters& InParameters, FShaderCompilerEnvironment& OutEnvironment);
+		static bool ShouldCompilePermutation(const FGlobalShaderPermutationParameters& Parameters);
 	};
 
 	class NNEDENOISERSHADERS_API FAutoExposureReduceFinalCS : public FGlobalShader
@@ -93,6 +96,7 @@ namespace UE::NNEDenoiserShaders::Internal
 		END_SHADER_PARAMETER_STRUCT()
 
 		static void ModifyCompilationEnvironment(const FGlobalShaderPermutationParameters& InParameters, FShaderCompilerEnvironment& OutEnvironment);
+		static bool ShouldCompilePermutation(const FGlobalShaderPermutationParameters& Parameters);
 	};
 
 } // namespace UE::NNEDenoiser::Private
\ No newline at end of file
diff --git a/Engine/Plugins/Experimental/NNEDenoiser/Source/NNEDenoiserShaders/Internal/NNEDenoiserShadersDefaultIOProcessCS.h b/Engine/Plugins/Experimental/NNEDenoiser/Source/NNEDenoiserShaders/Internal/NNEDenoiserShadersDefaultIOProcessCS.h
--- a/Engine/Plugins/Experimental/NNEDenoiser/Source/NNEDenoiserShaders/Internal/NNEDenoiserShadersDefaultIOProcessCS.h
+++ b/Engine/Plugins/Experimental/NNEDenoiser/Source/NNEDenoiserShaders/Internal/NNEDenoiserShadersDefaultIOProcessCS.h
@@ -45,6 +45,7 @@ namespace UE::NNEDenoiserShaders::Internal
 		END_SHADER_PARAMETER_STRUCT()
 
 		static void ModifyCompilationEnvironment(const FGlobalShaderPermutationParameters& InParameters, FShaderCompilerEnvironment& OutEnvironment);
+		static bool ShouldCompilePermutation(const FGlobalShaderPermutationParameters& Parameters);
 	};
 
 } // namespace UE::NNEDenoiser::Private
\ No newline at end of file
diff --git a/Engine/Plugins/Experimental/NNEDenoiser/Source/NNEDenoiserShaders/Internal/NNEDenoiserShadersMappedCopyCS.h b/Engine/Plugins/Experimental/NNEDenoiser/Source/NNEDenoiserShaders/Internal/NNEDenoiserShadersMappedCopyCS.h
--- a/Engine/Plugins/Experimental/NNEDenoiser/Source/NNEDenoiserShaders/Internal/NNEDenoiserShadersMappedCopyCS.h
+++ b/Engine/Plugins/Experimental/NNEDenoiser/Source/NNEDenoiserShaders/Internal/NNEDenoiserShadersMappedCopyCS.h
@@ -46,6 +46,7 @@ namespace UE::NNEDenoiserShaders::Internal
 		END_SHADER_PARAMETER_STRUCT()
 
 		static void ModifyCompilationEnvironment(const FGlobalShaderPermutationParameters& InParameters, FShaderCompilerEnvironment& OutEnvironment);
+		static bool ShouldCompilePermutation(const FGlobalShaderPermutationParameters& Parameters);
 	};
 
 	class NNEDENOISERSHADERS_API FNNEDenoiserBufferMappedCopyCS : public FGlobalShader
@@ -68,6 +69,7 @@ namespace UE::NNEDenoiserShaders::Internal
 		END_SHADER_PARAMETER_STRUCT()
 
 		static void ModifyCompilationEnvironment(const FGlobalShaderPermutationParameters& InParameters, FShaderCompilerEnvironment& OutEnvironment);
+		static bool ShouldCompilePermutation(const FGlobalShaderPermutationParameters& Parameters);
 	};
 
 	class NNEDENOISERSHADERS_API FNNEDenoiserBufferTextureMappedCopyCS : public FGlobalShader
@@ -90,6 +92,7 @@ namespace UE::NNEDenoiserShaders::Internal
 		END_SHADER_PARAMETER_STRUCT()
 
 		static void ModifyCompilationEnvironment(const FGlobalShaderPermutationParameters& InParameters, FShaderCompilerEnvironment& OutEnvironment);
+		static bool ShouldCompilePermutation(const FGlobalShaderPermutationParameters& Parameters);
 	};
 
 	class NNEDENOISERSHADERS_API FNNEDenoiserTextureMappedCopyCS : public FGlobalShader
@@ -112,6 +115,7 @@ namespace UE::NNEDenoiserShaders::Internal
 		END_SHADER_PARAMETER_STRUCT()
 
 		static void ModifyCompilationEnvironment(const FGlobalShaderPermutationParameters& InParameters, FShaderCompilerEnvironment& OutEnvironment);
+		static bool ShouldCompilePermutation(const FGlobalShaderPermutationParameters& Parameters);
 	};
 
 } // namespace UE::NNEDenoiser::Private
\ No newline at end of file
diff --git a/Engine/Plugins/Experimental/NNEDenoiser/Source/NNEDenoiserShaders/Internal/NNEDenoiserShadersTransferFunctionOidnCS.h b/Engine/Plugins/Experimental/NNEDenoiser/Source/NNEDenoiserShaders/Internal/NNEDenoiserShadersTransferFunctionOidnCS.h
--- a/Engine/Plugins/Experimental/NNEDenoiser/Source/NNEDenoiserShaders/Internal/NNEDenoiserShadersTransferFunctionOidnCS.h
+++ b/Engine/Plugins/Experimental/NNEDenoiser/Source/NNEDenoiserShaders/Internal/NNEDenoiserShadersTransferFunctionOidnCS.h
@@ -45,6 +45,7 @@ namespace UE::NNEDenoiserShaders::Internal
 		END_SHADER_PARAMETER_STRUCT()
 
 		static void ModifyCompilationEnvironment(const FGlobalShaderPermutationParameters& InParameters, FShaderCompilerEnvironment& OutEnvironment);
+		static bool ShouldCompilePermutation(const FGlobalShaderPermutationParameters& Parameters);
 	};
 
 } // namespace UE::NNEDenoiser::Private
\ No newline at end of file
diff --git a/Engine/Plugins/Experimental/NNEDenoiser/Source/NNEDenoiserShaders/Private/NNEDenoiserShadersAutoExposureCS.cpp b/Engine/Plugins/Experimental/NNEDenoiser/Source/NNEDenoiserShaders/Private/NNEDenoiserShadersAutoExposureCS.cpp
--- a/Engine/Plugins/Experimental/NNEDenoiser/Source/NNEDenoiserShaders/Private/NNEDenoiserShadersAutoExposureCS.cpp
+++ b/Engine/Plugins/Experimental/NNEDenoiser/Source/NNEDenoiserShaders/Private/NNEDenoiserShadersAutoExposureCS.cpp
@@ -31,8 +31,27 @@ namespace UE::NNEDenoiserShaders::Internal
 		OutEnvironment.SetDefine(TEXT("KEY"), FAutoExposureReduceConstants::KEY);
 	}
 
-	IMPLEMENT_GLOBAL_SHADER(FAutoExposureCS, "/NNEDenoiserShaders/NNEDenoiserShadersAutoExposure.usf", "AutoExposure", SF_Compute);
+	bool FAutoExposureCS::ShouldCompilePermutation(const FGlobalShaderPermutationParameters& Parameters)
+	{
+		return ShouldCompileRayTracingShadersForProject(Parameters.Platform);
+	}
 
+	bool FAutoExposureDownsampleCS::ShouldCompilePermutation(const FGlobalShaderPermutationParameters& Parameters)
+	{
+		return ShouldCompileRayTracingShadersForProject(Parameters.Platform);
+	}
+
+	bool FAutoExposureReduceCS::ShouldCompilePermutation(const FGlobalShaderPermutationParameters& Parameters)
+	{
+		return ShouldCompileRayTracingShadersForProject(Parameters.Platform);
+	}
+
+	bool FAutoExposureReduceFinalCS::ShouldCompilePermutation(const FGlobalShaderPermutationParameters& Parameters)
+	{
+		return ShouldCompileRayTracingShadersForProject(Parameters.Platform);
+	}
+
+	IMPLEMENT_GLOBAL_SHADER(FAutoExposureCS, "/NNEDenoiserShaders/NNEDenoiserShadersAutoExposure.usf", "AutoExposure", SF_Compute);
 	IMPLEMENT_GLOBAL_SHADER(FAutoExposureDownsampleCS, "/NNEDenoiserShaders/NNEDenoiserShadersAutoExposureDownsample.usf", "Downsample", SF_Compute);
 	IMPLEMENT_GLOBAL_SHADER(FAutoExposureReduceCS, "/NNEDenoiserShaders/NNEDenoiserShadersAutoExposureReduce.usf", "Reduce", SF_Compute);
 	IMPLEMENT_GLOBAL_SHADER(FAutoExposureReduceFinalCS, "/NNEDenoiserShaders/NNEDenoiserShadersAutoExposureReduceFinal.usf", "ReduceFinal", SF_Compute);
diff --git a/Engine/Plugins/Experimental/NNEDenoiser/Source/NNEDenoiserShaders/Private/NNEDenoiserShadersDefaultIOProcessCS.cpp b/Engine/Plugins/Experimental/NNEDenoiser/Source/NNEDenoiserShaders/Private/NNEDenoiserShadersDefaultIOProcessCS.cpp
--- a/Engine/Plugins/Experimental/NNEDenoiser/Source/NNEDenoiserShaders/Private/NNEDenoiserShadersDefaultIOProcessCS.cpp
+++ b/Engine/Plugins/Experimental/NNEDenoiser/Source/NNEDenoiserShaders/Private/NNEDenoiserShadersDefaultIOProcessCS.cpp
@@ -11,6 +11,11 @@ namespace UE::NNEDenoiserShaders::Internal
 		OutEnvironment.SetDefine(TEXT("MAX_FLT"), MAX_FLT);
 	}
 
+	bool FDefaultIOProcessCS::ShouldCompilePermutation(const FGlobalShaderPermutationParameters& Parameters)
+	{
+		return ShouldCompileRayTracingShadersForProject(Parameters.Platform);
+	}
+
 	IMPLEMENT_GLOBAL_SHADER(FDefaultIOProcessCS, "/NNEDenoiserShaders/NNEDenoiserShadersDefaultIOProcess.usf", "IOProcess", SF_Compute);
 
 } // UE::NNEDenoiser::Private
\ No newline at end of file
diff --git a/Engine/Plugins/Experimental/NNEDenoiser/Source/NNEDenoiserShaders/Private/NNEDenoiserShadersMappedCopyCS.cpp b/Engine/Plugins/Experimental/NNEDenoiser/Source/NNEDenoiserShaders/Private/NNEDenoiserShadersMappedCopyCS.cpp
--- a/Engine/Plugins/Experimental/NNEDenoiser/Source/NNEDenoiserShaders/Private/NNEDenoiserShadersMappedCopyCS.cpp
+++ b/Engine/Plugins/Experimental/NNEDenoiser/Source/NNEDenoiserShaders/Private/NNEDenoiserShadersMappedCopyCS.cpp
@@ -39,6 +39,26 @@ namespace UE::NNEDenoiserShaders::Internal
 		OutEnvironment.SetDefine(TEXT("OUTPUT_TYPE"), 0);
 	}
 
+	bool FNNEDenoiserTextureBufferMappedCopyCS::ShouldCompilePermutation(const FGlobalShaderPermutationParameters& Parameters)
+	{
+		return ShouldCompileRayTracingShadersForProject(Parameters.Platform);
+	}
+
+	bool FNNEDenoiserBufferMappedCopyCS::ShouldCompilePermutation(const FGlobalShaderPermutationParameters& Parameters)
+	{
+		return ShouldCompileRayTracingShadersForProject(Parameters.Platform);
+	}
+
+	bool FNNEDenoiserBufferTextureMappedCopyCS::ShouldCompilePermutation(const FGlobalShaderPermutationParameters& Parameters)
+	{
+		return ShouldCompileRayTracingShadersForProject(Parameters.Platform);
+	}
+
+	bool FNNEDenoiserTextureMappedCopyCS::ShouldCompilePermutation(const FGlobalShaderPermutationParameters& Parameters)
+	{
+		return ShouldCompileRayTracingShadersForProject(Parameters.Platform);
+	}
+
 	IMPLEMENT_GLOBAL_SHADER(FNNEDenoiserTextureBufferMappedCopyCS, "/NNEDenoiserShaders/NNEDenoiserShadersMappedCopy.usf", "MappedCopy", SF_Compute);
 	IMPLEMENT_GLOBAL_SHADER(FNNEDenoiserBufferMappedCopyCS, "/NNEDenoiserShaders/NNEDenoiserShadersMappedCopy.usf", "MappedCopy", SF_Compute);
 
diff --git a/Engine/Plugins/Experimental/NNEDenoiser/Source/NNEDenoiserShaders/Private/NNEDenoiserTransferFunctionOidn.cpp b/Engine/Plugins/Experimental/NNEDenoiser/Source/NNEDenoiserShaders/Private/NNEDenoiserTransferFunctionOidn.cpp
--- a/Engine/Plugins/Experimental/NNEDenoiser/Source/NNEDenoiserShaders/Private/NNEDenoiserTransferFunctionOidn.cpp
+++ b/Engine/Plugins/Experimental/NNEDenoiser/Source/NNEDenoiserShaders/Private/NNEDenoiserTransferFunctionOidn.cpp
@@ -11,6 +11,11 @@ namespace UE::NNEDenoiserShaders::Internal
 		OutEnvironment.SetDefine(TEXT("MAX_FLT"), MAX_FLT);
 	}
 
+	bool FTransferFunctionOidnCS::ShouldCompilePermutation(const FGlobalShaderPermutationParameters& Parameters)
+	{
+		return ShouldCompileRayTracingShadersForProject(Parameters.Platform);
+	}
+
 	IMPLEMENT_GLOBAL_SHADER(FTransferFunctionOidnCS, "/NNEDenoiserShaders/NNEDenoiserShadersTransferFunctionOidn.usf", "PreOrPostprocess", SF_Compute);
 
 } // UE::NNEDenoiser::Private
\ No newline at end of file
diff --git a/Engine/Plugins/Experimental/NNEDenoiser/Source/NNEDenoiserShaders/Shaders/NNEDenoiserShadersDefaultIOProcess.usf b/Engine/Plugins/Experimental/NNEDenoiser/Source/NNEDenoiserShaders/Shaders/NNEDenoiserShadersDefaultIOProcess.usf
--- a/Engine/Plugins/Experimental/NNEDenoiser/Source/NNEDenoiserShaders/Shaders/NNEDenoiserShadersDefaultIOProcess.usf
+++ b/Engine/Plugins/Experimental/NNEDenoiser/Source/NNEDenoiserShaders/Shaders/NNEDenoiserShadersDefaultIOProcess.usf
@@ -39,6 +39,5 @@ void IOProcess(in const uint3 DispatchThreadID : SV_DispatchThreadID)
 
 	const float4 Color = InputTexture.Load(int3(DispatchThreadID.x, DispatchThreadID.y, 0));
 
-	OutputTexture[DispatchThreadID.xy].rgb = GetProcessedValue(Color.rgb);
-	OutputTexture[DispatchThreadID.xy].a = Color.a;
+	OutputTexture[DispatchThreadID.xy] = float4(GetProcessedValue(Color.rgb), Color.a);
 }
\ No newline at end of file
diff --git a/Engine/Plugins/Experimental/NNEDenoiser/Source/NNEDenoiserShaders/Shaders/NNEDenoiserShadersTransferFunctionOidn.usf b/Engine/Plugins/Experimental/NNEDenoiser/Source/NNEDenoiserShaders/Shaders/NNEDenoiserShadersTransferFunctionOidn.usf
--- a/Engine/Plugins/Experimental/NNEDenoiser/Source/NNEDenoiserShaders/Shaders/NNEDenoiserShadersTransferFunctionOidn.usf
+++ b/Engine/Plugins/Experimental/NNEDenoiser/Source/NNEDenoiserShaders/Shaders/NNEDenoiserShadersTransferFunctionOidn.usf
@@ -76,6 +76,5 @@ void PreOrPostprocess(in const uint3 DispatchThreadID : SV_DispatchThreadID)
 
 	const float4 Color = InputTexture.Load(int3(DispatchThreadID.x, DispatchThreadID.y, 0));
 
-	OutputTexture[DispatchThreadID.xy].rgb = GetProcessedValue(Color.rgb);
-	OutputTexture[DispatchThreadID.xy].a = Color.a;
+	OutputTexture[DispatchThreadID.xy] = float4(GetProcessedValue(Color.rgb), Color.a);
 }
\ No newline at end of file
