Path Tracer: Enable NNEDenoiser and NFORDenoise plugins by default

Also change the path tracing denoiser modules to ClientOnly so that server builds do not try to load these plugins.

#rnx

diff --git a/Engine/Plugins/Experimental/NFORDenoise/NFORDenoise.uplugin b/Engine/Plugins/Experimental/NFORDenoise/NFORDenoise.uplugin
--- a/Engine/Plugins/Experimental/NFORDenoise/NFORDenoise.uplugin
+++ b/Engine/Plugins/Experimental/NFORDenoise/NFORDenoise.uplugin
@@ -10,7 +10,7 @@
 	"DocsURL": "",
 	"MarketplaceURL": "",
 	"SupportURL": "",
-	"EnabledByDefault": false,
+	"EnabledByDefault": true,
 	"CanContainContent": false,
 	"IsBetaVersion": false,
 	"IsExperimentalVersion": true,
@@ -22,7 +22,7 @@
 	"Modules": [
 		{
 			"Name": "NFORDenoise",
-			"Type": "RuntimeAndProgram",
+			"Type": "ClientOnly",
 			"LoadingPhase": "PostConfigInit",
 			"PlatformAllowList": [
 				"Win64",
diff --git a/Engine/Plugins/Experimental/NNEDenoiser/NNEDenoiser.uplugin b/Engine/Plugins/Experimental/NNEDenoiser/NNEDenoiser.uplugin
--- a/Engine/Plugins/Experimental/NNEDenoiser/NNEDenoiser.uplugin
+++ b/Engine/Plugins/Experimental/NNEDenoiser/NNEDenoiser.uplugin
@@ -10,19 +10,23 @@
 	"DocsURL": "",
 	"MarketplaceURL": "",
 	"SupportURL": "",
+	"EnabledByDefault" : true,
 	"CanContainContent": true,
 	"IsBetaVersion": false,
 	"IsExperimentalVersion": true,
 	"Installed": false,
+	"SupportedTargetPlatforms": [
+		"Win64", "Linux", "Mac"
+	],
 	"Modules": [
 		{
 			"Name": "NNEDenoiser",
-			"Type": "Runtime",
+			"Type": "ClientOnly",
 			"LoadingPhase": "Default"
 		},
 		{
 			"Name": "NNEDenoiserShaders",
-			"Type": "Runtime",
+			"Type": "ClientOnly",
 			"LoadingPhase": "PostConfigInit"
 		}
 	],
diff --git a/Engine/Plugins/Experimental/OpenImageDenoise/OpenImageDenoise.uplugin b/Engine/Plugins/Experimental/OpenImageDenoise/OpenImageDenoise.uplugin
--- a/Engine/Plugins/Experimental/OpenImageDenoise/OpenImageDenoise.uplugin
+++ b/Engine/Plugins/Experimental/OpenImageDenoise/OpenImageDenoise.uplugin
@@ -18,7 +18,7 @@
 	[
 		{
 			"Name" : "OpenImageDenoise",
-			"Type" : "RuntimeAndProgram",
+			"Type" : "ClientOnly",
 			"LoadingPhase" : "Default",
 			"PlatformAllowList": [
 				"Win64"
diff --git a/Engine/Source/Runtime/Renderer/Private/PathTracingSpatialTemporalDenoising.cpp b/Engine/Source/Runtime/Renderer/Private/PathTracingSpatialTemporalDenoising.cpp
--- a/Engine/Source/Runtime/Renderer/Private/PathTracingSpatialTemporalDenoising.cpp
+++ b/Engine/Source/Runtime/Renderer/Private/PathTracingSpatialTemporalDenoising.cpp
@@ -176,7 +176,7 @@ namespace {
 
 	TAutoConsoleVariable<FString> CVarPathTracingDenoiserName(
 		TEXT("r.PathTracing.Denoiser.Name"),
-		"OIDN",
+		"NNE_OIDN",
 		TEXT("Set the spatial denoiser name. It is the corresponding name registered by a denoiser plugin\n")
 		TEXT("Any registered denoiser should be able to denoise a single frame spatially.\n")
 	);
