NNEDenoiser: Make sure shader path is valid before adding it

#rnx

diff --git a/Engine/Plugins/Experimental/NNEDenoiser/Source/NNEDenoiserShaders/Private/NNEDenoiserShadersModule.cpp b/Engine/Plugins/Experimental/NNEDenoiser/Source/NNEDenoiserShaders/Private/NNEDenoiserShadersModule.cpp
--- a/Engine/Plugins/Experimental/NNEDenoiser/Source/NNEDenoiserShaders/Private/NNEDenoiserShadersModule.cpp
+++ b/Engine/Plugins/Experimental/NNEDenoiser/Source/NNEDenoiserShaders/Private/NNEDenoiserShadersModule.cpp
@@ -30,7 +30,10 @@ public:
 
 		FString ModuleShaderDir = FPaths::Combine(BaseDir, TEXT("Shaders"));
 		
-		AddShaderSourceDirectoryMapping(TEXT("/NNEDenoiserShaders"), ModuleShaderDir);
+		if (FPaths::DirectoryExists(ModuleShaderDir))
+		{
+			AddShaderSourceDirectoryMapping(TEXT("/NNEDenoiserShaders"), ModuleShaderDir);
+		}
 	}
 
 	virtual void ShutdownModule() override
diff --git a/Engine/Plugins/Experimental/NNERuntimeRDG/Source/NNEHlslShaders/Private/NNEHlslShadersModule.cpp b/Engine/Plugins/Experimental/NNERuntimeRDG/Source/NNEHlslShaders/Private/NNEHlslShadersModule.cpp
--- a/Engine/Plugins/Experimental/NNERuntimeRDG/Source/NNEHlslShaders/Private/NNEHlslShadersModule.cpp
+++ b/Engine/Plugins/Experimental/NNERuntimeRDG/Source/NNEHlslShaders/Private/NNEHlslShadersModule.cpp
@@ -28,7 +28,10 @@ public:
 
 		FString ModuleShaderDir = FPaths::Combine(BaseDir, TEXT("Shaders"));
 		
-		AddShaderSourceDirectoryMapping(TEXT("/NNE"), ModuleShaderDir);
+		if (FPaths::DirectoryExists(ModuleShaderDir))
+		{
+			AddShaderSourceDirectoryMapping(TEXT("/NNE"), ModuleShaderDir);
+		}
 	}
 
 	virtual void ShutdownModule() override
