[NNEDenoiser] Set as default and fix build issues.

#tests Editor, Win64 run NNEDenoiser test project

#rnx

diff --git a/Engine/Plugins/Experimental/NNEDenoiser/NNEDenoiser.uplugin b/Engine/Plugins/Experimental/NNEDenoiser/NNEDenoiser.uplugin
--- a/Engine/Plugins/Experimental/NNEDenoiser/NNEDenoiser.uplugin
+++ b/Engine/Plugins/Experimental/NNEDenoiser/NNEDenoiser.uplugin
@@ -10,7 +10,7 @@
 	"DocsURL": "",
 	"MarketplaceURL": "",
 	"SupportURL": "",
-	"EnabledByDefault" : false,
+	"EnabledByDefault" : true,
 	"CanContainContent": true,
 	"IsBetaVersion": false,
 	"IsExperimentalVersion": true,
@@ -21,12 +21,12 @@
 	"Modules": [
 		{
 			"Name": "NNEDenoiser",
-			"Type": "ClientOnly",
+			"Type": "Runtime",
 			"LoadingPhase": "Default"
 		},
 		{
 			"Name": "NNEDenoiserShaders",
-			"Type": "ClientOnly",
+			"Type": "Runtime",
 			"LoadingPhase": "PostConfigInit"
 		}
 	],
