[NNEDenosier] Try fix shader source directory mapping issue.

#tests Editor, Win64 run NNEDenosier

#rnx

diff --git a/Engine/Plugins/Experimental/NNEDenoiser/NNEDenoiser.uplugin b/Engine/Plugins/Experimental/NNEDenoiser/NNEDenoiser.uplugin
--- a/Engine/Plugins/Experimental/NNEDenoiser/NNEDenoiser.uplugin
+++ b/Engine/Plugins/Experimental/NNEDenoiser/NNEDenoiser.uplugin
@@ -34,10 +34,6 @@
 		{
 			"Name": "NNERuntimeORT",
 			"Enabled": true
-		},
-		{
-			"Name": "NNERuntimeRDG",
-			"Enabled": true
 		}
 	]
 }
\ No newline at end of file
diff --git a/Engine/Plugins/Experimental/NNEDenoiser/Shaders/NNEDenoiserShadersAutoExposure.usf b/Engine/Plugins/Experimental/NNEDenoiser/Shaders/NNEDenoiserShadersAutoExposure.usf
new file mode 100644
--- /dev/null
+++ b/Engine/Plugins/Experimental/NNEDenoiser/Shaders/NNEDenoiserShadersAutoExposure.usf
@@ -0,0 +1,40 @@
+// Copyright Epic Games, Inc. All Rights Reserved.
+
+#include "/Engine/Public/Platform.ush"
+#include "/Engine/Private/Common.ush"
+
+int InputTextureWidth;
+int InputTextureHeight;
+Texture2D InputTexture;
+RWBuffer<float> OutputBuffer;
+
+[numthreads(1, 1, 1)]
+void AutoExposure(in const uint3 DispatchThreadID : SV_DispatchThreadID)
+{
+	const float Key = 0.18;
+
+	float Sum = 0.0;
+	int Count = 0;
+	for (int i = 0; i < InputTextureWidth; i++)
+	{
+		for (int j = 0; j < InputTextureHeight; j++)
+		{
+			float3 Color = InputTexture.Load(int3(i, j, 0)).rgb;
+
+			Color = clamp(Color, float3(0.0, 0.0, 0.0), float3(MAX_FLT, MAX_FLT, MAX_FLT));
+
+			float Lum = Luminance(Color);
+			if (Lum > 1e-8)
+			{
+				Sum += log2(Lum);
+				Count++;
+			}
+		}
+	}
+
+	float Result = Count > 0 ? Sum / (float)Count : 1.0;
+	Result = Key / exp2(Result);
+
+	OutputBuffer[0] = Result;
+	OutputBuffer[1] = 1.0 / Result;
+}
\ No newline at end of file
diff --git a/Engine/Plugins/Experimental/NNEDenoiser/Shaders/NNEDenoiserShadersAutoExposureDownsample.usf b/Engine/Plugins/Experimental/NNEDenoiser/Shaders/NNEDenoiserShadersAutoExposureDownsample.usf
new file mode 100644
--- /dev/null
+++ b/Engine/Plugins/Experimental/NNEDenoiser/Shaders/NNEDenoiserShadersAutoExposureDownsample.usf
@@ -0,0 +1,55 @@
+// Copyright Epic Games, Inc. All Rights Reserved.
+
+#include "/Engine/Public/Platform.ush"
+#include "/Engine/Private/Common.ush"
+
+int InputTextureWidth;
+int InputTextureHeight;
+Texture2D InputTexture;
+
+int NumBinsW;
+int NumBinsH;
+RWBuffer<float> OutputBins;
+
+groupshared float Sums[THREAD_GROUP_SIZE * THREAD_GROUP_SIZE];
+
+[numthreads(THREAD_GROUP_SIZE, THREAD_GROUP_SIZE, 1)]
+void Downsample(in const uint3 GroupThreadID : SV_GroupThreadID, in const uint3 GroupID : SV_GroupID, in const uint3 DispatchThreadID : SV_DispatchThreadID)
+{
+	const int BeginW = GroupID.x * InputTextureWidth / NumBinsW;
+	const int BeginH = GroupID.y * InputTextureHeight / NumBinsH;
+	const int EndW = (GroupID.x + 1) * InputTextureWidth / NumBinsW;
+	const int EndH = (GroupID.y + 1) * InputTextureHeight / NumBinsH;
+
+	const int CoordW = BeginW + GroupThreadID.x;
+	const int CoordH = BeginH + GroupThreadID.y;
+
+	float Lum = 0.0;
+	if (CoordW < EndW && CoordH < EndH)
+	{
+		float3 Color = InputTexture.Load(int3(CoordW, CoordH, 0)).rgb;
+		Color = clamp(Color, float3(0.0, 0.0, 0.0), float3(MAX_FLT, MAX_FLT, MAX_FLT));
+
+		Lum = Luminance(Color);
+	}
+
+	const int LocalID = GroupThreadID.y * THREAD_GROUP_SIZE + GroupThreadID.x;
+	Sums[LocalID] = Lum;
+
+	for (int i = THREAD_GROUP_SIZE * THREAD_GROUP_SIZE / 2; i > 0; i >>= 1)
+	{
+		GroupMemoryBarrierWithGroupSync();
+		
+		if (LocalID < i)
+		{
+			Sums[LocalID] += Sums[LocalID + i];
+		}
+	}
+
+	if (LocalID == 0)
+	{
+		const float AverageLum = Sums[0] / float((EndW - BeginW) * (EndH - BeginH));
+		const int GroupLinearID = GroupID.y * NumBinsW + GroupID.x;
+		OutputBins[GroupLinearID] = AverageLum;
+	}
+}
\ No newline at end of file
diff --git a/Engine/Plugins/Experimental/NNEDenoiser/Shaders/NNEDenoiserShadersAutoExposureReduce.usf b/Engine/Plugins/Experimental/NNEDenoiser/Shaders/NNEDenoiserShadersAutoExposureReduce.usf
new file mode 100644
--- /dev/null
+++ b/Engine/Plugins/Experimental/NNEDenoiser/Shaders/NNEDenoiserShadersAutoExposureReduce.usf
@@ -0,0 +1,50 @@
+// Copyright Epic Games, Inc. All Rights Reserved.
+
+#include "/Engine/Public/Platform.ush"
+#include "/Engine/Private/Common.ush"
+
+int InputSize;
+RWBuffer<float> InputBins;
+RWBuffer<float> OutputSums;
+RWBuffer<int> OutputCounts;
+int NumThreads;
+
+groupshared float Sums[THREAD_GROUP_SIZE];
+groupshared int Counts[THREAD_GROUP_SIZE];
+
+[numthreads(THREAD_GROUP_SIZE, 1, 1)]
+void Reduce(in const uint3 GroupThreadID : SV_GroupThreadID, in const uint3 GroupID : SV_GroupID, in const uint3 DispatchThreadID : SV_DispatchThreadID)
+{
+	float LocalSum = 0.0;
+	int LocalCount = 0;
+	for (int i = DispatchThreadID.x; i < InputSize; i += NumThreads)
+	{
+		const float Lum = InputBins[i];
+		if (Lum > EPS)
+		{
+			LocalSum += log2(Lum);
+			++LocalCount;
+		}
+	}
+
+	const int LocalID = GroupThreadID.x;
+	Sums[LocalID] = LocalSum;
+	Counts[LocalID] = LocalCount;
+
+	for (int i = THREAD_GROUP_SIZE / 2; i > 0; i >>= 1)
+	{
+		GroupMemoryBarrierWithGroupSync();
+
+		if (LocalID < i)
+		{
+			Sums[LocalID] += Sums[LocalID + i];
+			Counts[LocalID] += Counts[LocalID + i];
+		}
+	}
+
+	if (LocalID == 0)
+	{
+		OutputSums[GroupID.x] = Sums[0];
+		OutputCounts[GroupID.x] = Counts[0];
+	}
+}
\ No newline at end of file
diff --git a/Engine/Plugins/Experimental/NNEDenoiser/Shaders/NNEDenoiserShadersAutoExposureReduceFinal.usf b/Engine/Plugins/Experimental/NNEDenoiser/Shaders/NNEDenoiserShadersAutoExposureReduceFinal.usf
new file mode 100644
--- /dev/null
+++ b/Engine/Plugins/Experimental/NNEDenoiser/Shaders/NNEDenoiserShadersAutoExposureReduceFinal.usf
@@ -0,0 +1,46 @@
+// Copyright Epic Games, Inc. All Rights Reserved.
+
+#include "/Engine/Public/Platform.ush"
+#include "/Engine/Private/Common.ush"
+
+int InputSize;
+RWBuffer<float> InputSums;
+RWBuffer<int> InputCounts;
+RWBuffer<float> OutputBuffer; // 0: Result, 1: 1.0 / Result
+
+groupshared float Sums[THREAD_GROUP_SIZE];
+groupshared int Counts[THREAD_GROUP_SIZE];
+
+[numthreads(THREAD_GROUP_SIZE, 1, 1)]
+void ReduceFinal(in const uint3 GroupThreadID : SV_GroupThreadID)
+{
+	const int LocalID = GroupThreadID.x;
+
+	if (LocalID < InputSize)
+	{
+		Sums[LocalID] = InputSums[LocalID];
+		Counts[LocalID] = InputCounts[LocalID];
+	}
+	else
+	{
+		Sums[LocalID] = 0;
+		Counts[LocalID] = 0;
+	}
+
+	for (int i = THREAD_GROUP_SIZE / 2; i > 0; i >>= 1)
+	{
+		GroupMemoryBarrierWithGroupSync();
+
+		if (LocalID < i)
+		{
+			Sums[LocalID] += Sums[LocalID + i];
+			Counts[LocalID] += Counts[LocalID + i];
+		}
+	}
+
+	if (LocalID == 0)
+	{
+		OutputBuffer[0] = (Counts[0] > 0) ? (KEY / exp2(Sums[0] / float(Counts[0]))) : 1.0;
+		OutputBuffer[1] = 1.0 / OutputBuffer[0];
+	}
+}
\ No newline at end of file
diff --git a/Engine/Plugins/Experimental/NNEDenoiser/Shaders/NNEDenoiserShadersDefaultIOProcess.usf b/Engine/Plugins/Experimental/NNEDenoiser/Shaders/NNEDenoiserShadersDefaultIOProcess.usf
new file mode 100644
--- /dev/null
+++ b/Engine/Plugins/Experimental/NNEDenoiser/Shaders/NNEDenoiserShadersDefaultIOProcess.usf
@@ -0,0 +1,43 @@
+// Copyright Epic Games, Inc. All Rights Reserved.
+
+#include "/Engine/Public/Platform.ush"
+#include "/Engine/Private/Common.ush"
+
+int InputTextureWidth;
+int InputTextureHeight;
+Texture2D InputTexture;
+
+int OutputTextureWidth;
+int OutputTextureHeight;
+RWTexture2D<float4> OutputTexture;
+
+float3 GetProcessedValue(float3 Value)
+{
+#if INPUT_KIND_INDEX == 0 // Input color
+	return clamp(Value, float3(0.0, 0.0, 0.0), float3(MAX_FLT, MAX_FLT, MAX_FLT));
+#elif INPUT_KIND_INDEX == 1 // Input albedo
+	return clamp(Value, float3(0.0, 0.0, 0.0), float3(1.0, 1.0, 1.0));
+#elif INPUT_KIND_INDEX == 2 // Input normal
+	return clamp(Value, float3(-1.0, -1.0, -1.0), float3(1.0, 1.0, 1.0)) * 0.5 + float3(0.5, 0.5, 0.5);
+#elif INPUT_KIND_INDEX == 3 // Input flow
+	// TODO flow texture check for any necessary conversion
+	return Value;
+#else
+	// INPUT_KIND_INDEX == 4 // Output color
+	return clamp(Value, float3(0.0, 0.0, 0.0), float3(MAX_FLT, MAX_FLT, MAX_FLT));
+#endif
+}
+
+
+[numthreads(THREAD_GROUP_SIZE, THREAD_GROUP_SIZE, 1)]
+void IOProcess(in const uint3 DispatchThreadID : SV_DispatchThreadID)
+{
+	if (DispatchThreadID.x >= OutputTextureWidth || DispatchThreadID.y >= OutputTextureHeight)
+	{
+		return;
+	}
+
+	const float4 Color = InputTexture.Load(int3(DispatchThreadID.x, DispatchThreadID.y, 0));
+
+	OutputTexture[DispatchThreadID.xy] = float4(GetProcessedValue(Color.rgb), Color.a);
+}
\ No newline at end of file
diff --git a/Engine/Plugins/Experimental/NNEDenoiser/Shaders/NNEDenoiserShadersMappedCopy.usf b/Engine/Plugins/Experimental/NNEDenoiser/Shaders/NNEDenoiserShadersMappedCopy.usf
new file mode 100644
--- /dev/null
+++ b/Engine/Plugins/Experimental/NNEDenoiser/Shaders/NNEDenoiserShadersMappedCopy.usf
@@ -0,0 +1,69 @@
+// Copyright Epic Games, Inc. All Rights Reserved.
+
+#include "/Engine/Public/Platform.ush"
+#include "/Engine/Private/Common.ush"
+
+#if BUFFER_TYPE_INDEX == 4 // Float16
+	#define BUFFER_TYPE half
+	#define READ(X) float(X)
+	#define WRITE(X) half(X)
+#else
+	// BUFFER_TYPE_INDEX == 3 // Float32
+	#define BUFFER_TYPE float
+	#define READ(X) X
+	#define WRITE(X) X
+#endif
+
+int Width;
+int Height;
+
+int4 OutputChannel_InputChannel_Unused_Unused[MAX_NUM_MAPPED_CHANNELS];
+
+#if INPUT_TYPE == 0 // Texture
+Texture2D InputTexture;
+#else // INPUT_TYPE == 1 // Buffer
+RWBuffer<BUFFER_TYPE> InputBuffer;
+#endif
+
+#if OUTPUT_TYPE == 0 // Texture
+RWTexture2D<float4> OutputTexture;
+#else // OUTPUT_TYPE == 1 // Buffer
+RWBuffer<BUFFER_TYPE> OutputBuffer;
+#endif
+
+[numthreads(THREAD_GROUP_SIZE, THREAD_GROUP_SIZE, 1)]
+void MappedCopy(in const uint3 DispatchThreadID : SV_DispatchThreadID)
+{
+	if (DispatchThreadID.x >= Width || DispatchThreadID.y >= Height)
+	{
+		return;
+	}
+
+#if INPUT_TYPE == 1 || OUTPUT_TYPE == 1
+	const int Idx = DispatchThreadID.y * Width + DispatchThreadID.x;
+	const int Offset = Width * Height;
+#endif
+
+#if INPUT_TYPE == 0
+	const float4 InputColor = InputTexture.Load(int3(DispatchThreadID.x, DispatchThreadID.y, 0));
+#endif
+
+	for (int I = 0; I < NUM_MAPPED_CHANNELS; I++)
+	{
+		const int4 Info = OutputChannel_InputChannel_Unused_Unused[I];
+		const int OutputChannel = Info.x;
+		const int InputChannel = Info.y;
+
+#if INPUT_TYPE == 0
+		const float InputValue = InputColor[InputChannel];
+#else
+		const float InputValue = InputBuffer[Idx + InputChannel * Offset];
+#endif
+
+#if OUTPUT_TYPE == 0
+		OutputTexture[DispatchThreadID.xy][OutputChannel] = READ(InputValue);
+#else
+		OutputBuffer[Idx + OutputChannel * Offset] = WRITE(InputValue);
+#endif
+	}
+}
\ No newline at end of file
diff --git a/Engine/Plugins/Experimental/NNEDenoiser/Shaders/NNEDenoiserShadersTransferFunctionOidn.usf b/Engine/Plugins/Experimental/NNEDenoiser/Shaders/NNEDenoiserShadersTransferFunctionOidn.usf
new file mode 100644
--- /dev/null
+++ b/Engine/Plugins/Experimental/NNEDenoiser/Shaders/NNEDenoiserShadersTransferFunctionOidn.usf
@@ -0,0 +1,80 @@
+// Copyright Epic Games, Inc. All Rights Reserved.
+
+#include "/Engine/Public/Platform.ush"
+#include "/Engine/Private/Common.ush"
+
+int InputTextureWidth;
+int InputTextureHeight;
+Texture2D InputTexture;
+
+int OutputTextureWidth;
+int OutputTextureHeight;
+RWTexture2D<float4> OutputTexture;
+
+RWBuffer<float> InputScaleBuffer;
+
+float NormScale;
+float InvNormScale;
+
+///
+
+// from [REDACTED_PROFILE_URL]/oidn/blob/master/core/color.h
+// struct FPU
+// {
+static float a  =  1.41283765e+03;
+static float b  =  1.64593172e+00;
+static float c  =  4.31384981e-01;
+static float d  = -2.94139609e-03;
+static float e  =  1.92653254e-01;
+static float f  =  6.26026094e-03;
+static float g  =  9.98620152e-01;
+static float y0 =  1.57945760e-06;
+static float y1 =  3.22087631e-02;
+static float x0 =  2.23151711e-03;
+static float x1 =  3.70974749e-01;
+
+static float Forward(float y)
+{
+	if (y <= y0) return a * y;
+	else if (y <= y1) return b * pow(y, c) + d;
+	return e * log(y + f) + g;
+}
+
+static float Inverse(float x)
+{
+	if (x <= x0) return x / a;
+	else if (x <= x1) return pow((x - d) / b, 1.0 / c);
+	return exp((x - g) / e) - f;
+}
+
+float3 GetProcessedValue(float3 Value)
+{
+#if MODE == 0 // Forward
+	Value = Value * InputScaleBuffer[0];
+	Value = clamp(Value, float3(0.0, 0.0, 0.0), float3(MAX_FLT, MAX_FLT, MAX_FLT));
+
+	return float3(Forward(Value.r), Forward(Value.g), Forward(Value.b)) * NormScale;
+#else
+	// MODE == 1 // Inverse
+	Value = clamp(Value, float3(0.0, 0.0, 0.0), float3(MAX_FLT, MAX_FLT, MAX_FLT)) * InvNormScale;
+	
+	return float3(Inverse(Value.r), Inverse(Value.g), Inverse(Value.b)) * InputScaleBuffer[1];
+#endif
+}
+
+// };
+
+///
+
+[numthreads(THREAD_GROUP_SIZE, THREAD_GROUP_SIZE, 1)]
+void PreOrPostprocess(in const uint3 DispatchThreadID : SV_DispatchThreadID)
+{
+	if (DispatchThreadID.x >= OutputTextureWidth || DispatchThreadID.y >= OutputTextureHeight)
+	{
+		return;
+	}
+
+	const float4 Color = InputTexture.Load(int3(DispatchThreadID.x, DispatchThreadID.y, 0));
+
+	OutputTexture[DispatchThreadID.xy] = float4(GetProcessedValue(Color.rgb), Color.a);
+}
\ No newline at end of file
diff --git a/Engine/Plugins/Experimental/NNEDenoiser/Source/NNEDenoiserShaders/Private/NNEDenoiserShadersModule.cpp b/Engine/Plugins/Experimental/NNEDenoiser/Source/NNEDenoiserShaders/Private/NNEDenoiserShadersModule.cpp
--- a/Engine/Plugins/Experimental/NNEDenoiser/Source/NNEDenoiserShaders/Private/NNEDenoiserShadersModule.cpp
+++ b/Engine/Plugins/Experimental/NNEDenoiser/Source/NNEDenoiserShaders/Private/NNEDenoiserShadersModule.cpp
@@ -15,25 +15,17 @@ public:
 	// Begin IModuleInterface
 	virtual void StartupModule() override
 	{
-		FString BaseDir;
-
 		const TSharedPtr<IPlugin> Plugin = IPluginManager::Get().FindPlugin(TEXT("NNEDenoiser"));
-		
 		if (Plugin.IsValid())
 		{
-			BaseDir = Plugin->GetBaseDir() + TEXT("/Source/NNEDenoiserShaders");
+			const FString ShadersDir = FPaths::Combine(Plugin->GetBaseDir(), TEXT("Shaders"));
+
+			AddShaderSourceDirectoryMapping(TEXT("/NNEDenoiserShaders"), ShadersDir);
 		}
 		else
 		{
 			UE_LOG(LogNNEDenoiserShaders, Warning, TEXT("Shaders directory not added. Failed to find NNEDenoiser plugin"));
 		}
-
-		FString ModuleShaderDir = FPaths::Combine(BaseDir, TEXT("Shaders"));
-		
-		if (FPaths::DirectoryExists(ModuleShaderDir))
-		{
-			AddShaderSourceDirectoryMapping(TEXT("/NNEDenoiserShaders"), ModuleShaderDir);
-		}
 	}
 
 	virtual void ShutdownModule() override
diff --git a/Engine/Plugins/Experimental/NNEDenoiser/Source/NNEDenoiserShaders/Shaders/NNEDenoiserShadersAutoExposure.usf b/Engine/Plugins/Experimental/NNEDenoiser/Source/NNEDenoiserShaders/Shaders/NNEDenoiserShadersAutoExposure.usf
deleted file mode 100644
diff --git a/Engine/Plugins/Experimental/NNEDenoiser/Source/NNEDenoiserShaders/Shaders/NNEDenoiserShadersAutoExposureDownsample.usf b/Engine/Plugins/Experimental/NNEDenoiser/Source/NNEDenoiserShaders/Shaders/NNEDenoiserShadersAutoExposureDownsample.usf
deleted file mode 100644
diff --git a/Engine/Plugins/Experimental/NNEDenoiser/Source/NNEDenoiserShaders/Shaders/NNEDenoiserShadersAutoExposureReduce.usf b/Engine/Plugins/Experimental/NNEDenoiser/Source/NNEDenoiserShaders/Shaders/NNEDenoiserShadersAutoExposureReduce.usf
deleted file mode 100644
diff --git a/Engine/Plugins/Experimental/NNEDenoiser/Source/NNEDenoiserShaders/Shaders/NNEDenoiserShadersAutoExposureReduceFinal.usf b/Engine/Plugins/Experimental/NNEDenoiser/Source/NNEDenoiserShaders/Shaders/NNEDenoiserShadersAutoExposureReduceFinal.usf
deleted file mode 100644
diff --git a/Engine/Plugins/Experimental/NNEDenoiser/Source/NNEDenoiserShaders/Shaders/NNEDenoiserShadersDefaultIOProcess.usf b/Engine/Plugins/Experimental/NNEDenoiser/Source/NNEDenoiserShaders/Shaders/NNEDenoiserShadersDefaultIOProcess.usf
deleted file mode 100644
diff --git a/Engine/Plugins/Experimental/NNEDenoiser/Source/NNEDenoiserShaders/Shaders/NNEDenoiserShadersMappedCopy.usf b/Engine/Plugins/Experimental/NNEDenoiser/Source/NNEDenoiserShaders/Shaders/NNEDenoiserShadersMappedCopy.usf
deleted file mode 100644
diff --git a/Engine/Plugins/Experimental/NNEDenoiser/Source/NNEDenoiserShaders/Shaders/NNEDenoiserShadersTransferFunctionOidn.usf b/Engine/Plugins/Experimental/NNEDenoiser/Source/NNEDenoiserShaders/Shaders/NNEDenoiserShadersTransferFunctionOidn.usf
deleted file mode 100644
