[NNEDenoiser] Only register View Extension when Ray Tracing is enabled.

#tests Editor, Win64 and Mac to verify view extension not registered

#rnx

diff --git a/Engine/Plugins/Experimental/NNEDenoiser/Source/NNEDenoiser/Private/NNEDenoise.cpp b/Engine/Plugins/Experimental/NNEDenoiser/Source/NNEDenoiser/Private/NNEDenoise.cpp
--- a/Engine/Plugins/Experimental/NNEDenoiser/Source/NNEDenoiser/Private/NNEDenoise.cpp
+++ b/Engine/Plugins/Experimental/NNEDenoiser/Source/NNEDenoiser/Private/NNEDenoise.cpp
@@ -9,21 +9,28 @@ DEFINE_LOG_CATEGORY(LogNNEDenoiser);
 
 void FNNEDenoiserModule::StartupModule()
 {
-	// This code will execute after your module is loaded into memory; the exact timing is specified in the .uplugin file per-module
-	UE_LOG(LogNNEDenoiser, Log, TEXT("NNEDenoiser module starting up"));
+	// During cook, we are not allowed to call IsRayTracingAllowed() later, therefore we assume Ray Tracing is not available and do not register the View Extension.
+	if (IsRunningCookCommandlet())
+	{
+		return;
+	}
 
 	FCoreDelegates::OnPostEngineInit.AddLambda([this] ()
 	{
-		ViewExtension = FSceneViewExtensions::NewExtension<UE::NNEDenoiser::Private::FViewExtension>();
+		// Only register View Extension if Ray Tracing is available.
+		if (IsRayTracingAllowed())
+		{
+			ViewExtension = FSceneViewExtensions::NewExtension<UE::NNEDenoiser::Private::FViewExtension>();
+		}
+		else
+		{
+			UE_LOG(LogNNEDenoiser, Log, TEXT("Ray Tracing is not enabled, therefore NNEDenoiser is not registered!"));
+		}
 	});
 }
 
 void FNNEDenoiserModule::ShutdownModule()
 {
-	// This function may be called during shutdown to clean up your module.  For modules that support dynamic reloading,
-	// we call this function before unloading the module.
-	UE_LOG(LogNNEDenoiser, Log, TEXT("NNEDenoiser module shut down"));
-
 	ViewExtension.Reset();
 }
 
