[NNEDenoiser] Implemented fallback system when creating model instance with NNE runtimes, mainly to cover platforms without DirectML available.

#tests Editor, Win64 and Mac verified fallback to available runtime

#rnx

diff --git a/Engine/Plugins/Experimental/NNEDenoiser/Source/NNEDenoiser/Private/NNEDenoiserModelInstanceCPU.cpp b/Engine/Plugins/Experimental/NNEDenoiser/Source/NNEDenoiser/Private/NNEDenoiserModelInstanceCPU.cpp
--- a/Engine/Plugins/Experimental/NNEDenoiser/Source/NNEDenoiser/Private/NNEDenoiserModelInstanceCPU.cpp
+++ b/Engine/Plugins/Experimental/NNEDenoiser/Source/NNEDenoiser/Private/NNEDenoiserModelInstanceCPU.cpp
@@ -20,40 +20,41 @@ END_SHADER_PARAMETER_STRUCT()
 namespace UE::NNEDenoiser::Private
 {
 
-	static const FString DefaultRuntimeCPUName = TEXT("NNERuntimeORTCpu");
-
-	TUniquePtr<FModelInstanceCPU> FModelInstanceCPU::Make(UNNEModelData& ModelData, const FString &RuntimeNameOverride)
+	TUniquePtr<FModelInstanceCPU> FModelInstanceCPU::Make(UNNEModelData& ModelData, const FString& RuntimeName)
 	{
-		const FString RuntimeCPUName = RuntimeNameOverride.IsEmpty() ? DefaultRuntimeCPUName : RuntimeNameOverride;
+		check(!RuntimeName.IsEmpty());
 
-		// Create the model
-		TWeakInterfacePtr<INNERuntimeCPU> RuntimeCPU = UE::NNE::GetRuntime<INNERuntimeCPU>(RuntimeCPUName);
-		if (!RuntimeCPU.IsValid())
+		TWeakInterfacePtr<INNERuntimeCPU> Runtime = UE::NNE::GetRuntime<INNERuntimeCPU>(RuntimeName);
+		if (!Runtime.IsValid())
 		{
-			UE_LOG(LogNNEDenoiser, Error, TEXT("No CPU runtime '%s' found. Valid CPU runtimes are: "), *RuntimeCPUName);
-			for (const FString& RuntimeName : UE::NNE::GetAllRuntimeNames<INNERuntimeCPU>())
+			UE_LOG(LogNNEDenoiser, Log, TEXT("Could not create model instance. No CPU runtime '%s' found. Valid CPU runtimes are: "), *RuntimeName);
+			for (const FString& Name : UE::NNE::GetAllRuntimeNames<INNERuntimeCPU>())
 			{
-				UE_LOG(LogNNEDenoiser, Error, TEXT("- %s"), *RuntimeName);
+				UE_LOG(LogNNEDenoiser, Log, TEXT("- %s"), *Name);
 			}
 			return {};
 		}
 
-		TSharedPtr<UE::NNE::IModelCPU> Model = RuntimeCPU->CreateModelCPU(ModelData);
+		if (Runtime->CanCreateModelCPU(ModelData) != INNERuntimeCPU::ECanCreateModelCPUStatus::Ok)
+		{
+			UE_LOG(LogNNEDenoiser, Log, TEXT("%s on CPU can not create model"), *RuntimeName);
+			return {};
+		}
+
+		TSharedPtr<UE::NNE::IModelCPU> Model = Runtime->CreateModelCPU(ModelData);
 		if (!Model.IsValid())
 		{
-			UE_LOG(LogNNEDenoiser, Error, TEXT("Could not create model using %s"), *RuntimeCPUName);
+			UE_LOG(LogNNEDenoiser, Log, TEXT("Could not create model using %s on CPU"), *RuntimeName);
 			return {};
 		}
 
 		TSharedPtr<UE::NNE::IModelInstanceCPU> ModelInstance = Model->CreateModelInstanceCPU();
 		if (!ModelInstance.IsValid())
 		{
-			UE_LOG(LogNNEDenoiser, Error, TEXT("Could not create model instance using %s"), *RuntimeCPUName);
+			UE_LOG(LogNNEDenoiser, Log, TEXT("Could not create model instance using %s on CPU"), *RuntimeName);
 			return {};
 		}
 
-		UE_LOG(LogNNEDenoiser, Log, TEXT("NNEDenoiserCPU: Creaded model instance from %s using %s"), *ModelData.GetFileId().ToString(), *RuntimeCPUName);
-
 		return MakeUnique<FModelInstanceCPU>(ModelInstance.ToSharedRef());
 	}
 
diff --git a/Engine/Plugins/Experimental/NNEDenoiser/Source/NNEDenoiser/Private/NNEDenoiserModelInstanceCPU.h b/Engine/Plugins/Experimental/NNEDenoiser/Source/NNEDenoiser/Private/NNEDenoiserModelInstanceCPU.h
--- a/Engine/Plugins/Experimental/NNEDenoiser/Source/NNEDenoiser/Private/NNEDenoiserModelInstanceCPU.h
+++ b/Engine/Plugins/Experimental/NNEDenoiser/Source/NNEDenoiser/Private/NNEDenoiserModelInstanceCPU.h
@@ -20,7 +20,7 @@ class FModelInstanceCPU : public IModelInstance
 public:
 	explicit FModelInstanceCPU(TSharedRef<UE::NNE::IModelInstanceCPU> ModelInstance);
 
-	static TUniquePtr<FModelInstanceCPU> Make(UNNEModelData& ModelData, const FString& RuntimeNameOverride = {});
+	static TUniquePtr<FModelInstanceCPU> Make(UNNEModelData& ModelData, const FString& RuntimeName);
 
 	virtual ~FModelInstanceCPU();
 
diff --git a/Engine/Plugins/Experimental/NNEDenoiser/Source/NNEDenoiser/Private/NNEDenoiserModelInstanceGPU.cpp b/Engine/Plugins/Experimental/NNEDenoiser/Source/NNEDenoiser/Private/NNEDenoiserModelInstanceGPU.cpp
--- a/Engine/Plugins/Experimental/NNEDenoiser/Source/NNEDenoiser/Private/NNEDenoiserModelInstanceGPU.cpp
+++ b/Engine/Plugins/Experimental/NNEDenoiser/Source/NNEDenoiser/Private/NNEDenoiserModelInstanceGPU.cpp
@@ -23,40 +23,41 @@ END_SHADER_PARAMETER_STRUCT()
 namespace UE::NNEDenoiser::Private
 {
 
-	static const FString DefaultRuntimeGPUName = TEXT("NNERuntimeORTDml");
-
-	TUniquePtr<FModelInstanceGPU> FModelInstanceGPU::Make(UNNEModelData& ModelData, const FString& RuntimeNameOverride)
+	TUniquePtr<FModelInstanceGPU> FModelInstanceGPU::Make(UNNEModelData& ModelData, const FString& RuntimeName)
 	{
-		const FString RuntimeGPUName = RuntimeNameOverride.IsEmpty() ? DefaultRuntimeGPUName : RuntimeNameOverride;
+		check(!RuntimeName.IsEmpty());
 
-		// Create the model
-		TWeakInterfacePtr<INNERuntimeGPU> RuntimeGPU = UE::NNE::GetRuntime<INNERuntimeGPU>(RuntimeGPUName);
-		if (!RuntimeGPU.IsValid())
+		TWeakInterfacePtr<INNERuntimeGPU> Runtime = UE::NNE::GetRuntime<INNERuntimeGPU>(RuntimeName);
+		if (!Runtime.IsValid())
 		{
-			UE_LOG(LogNNEDenoiser, Error, TEXT("No GPU runtime '%s' found. Valid GPU runtimes are: "), *RuntimeGPUName);
-			for (const FString& RuntimeName : UE::NNE::GetAllRuntimeNames<INNERuntimeGPU>())
+			UE_LOG(LogNNEDenoiser, Log, TEXT("Could not create model instance. No GPU runtime '%s' found. Valid GPU runtimes are: "), *RuntimeName);
+			for (const FString& Name : UE::NNE::GetAllRuntimeNames<INNERuntimeGPU>())
 			{
-				UE_LOG(LogNNEDenoiser, Error, TEXT("- %s"), *RuntimeName);
+				UE_LOG(LogNNEDenoiser, Log, TEXT("- %s"), *Name);
 			}
 			return {};
 		}
 
-		TSharedPtr<UE::NNE::IModelGPU> Model = RuntimeGPU->CreateModelGPU(ModelData);
+		if (Runtime->CanCreateModelGPU(ModelData) != INNERuntimeGPU::ECanCreateModelGPUStatus::Ok)
+		{
+			UE_LOG(LogNNEDenoiser, Log, TEXT("%s on GPU can not create model"), *RuntimeName);
+			return {};
+		}
+
+		TSharedPtr<UE::NNE::IModelGPU> Model = Runtime->CreateModelGPU(ModelData);
 		if (!Model.IsValid())
 		{
-			UE_LOG(LogNNEDenoiser, Error, TEXT("Could not create model using %s"), *RuntimeGPUName);
+			UE_LOG(LogNNEDenoiser, Log, TEXT("Could not create model using %s on GPU"), *RuntimeName);
 			return {};
 		}
 
 		TSharedPtr<UE::NNE::IModelInstanceGPU> ModelInstance = Model->CreateModelInstanceGPU();
 		if (!ModelInstance.IsValid())
 		{
-			UE_LOG(LogNNEDenoiser, Error, TEXT("Could not create model instance using %s"), *RuntimeGPUName);
+			UE_LOG(LogNNEDenoiser, Log, TEXT("Could not create model instance using %s on GPU"), *RuntimeName);
 			return {};
 		}
 
-		UE_LOG(LogNNEDenoiser, Log, TEXT("NNEDenoiserGPU: Creaded model instance from %s using %s"), *ModelData.GetFileId().ToString(), *RuntimeGPUName);
-
 		return MakeUnique<FModelInstanceGPU>(ModelInstance.ToSharedRef());
 	}
 
diff --git a/Engine/Plugins/Experimental/NNEDenoiser/Source/NNEDenoiser/Private/NNEDenoiserModelInstanceGPU.h b/Engine/Plugins/Experimental/NNEDenoiser/Source/NNEDenoiser/Private/NNEDenoiserModelInstanceGPU.h
--- a/Engine/Plugins/Experimental/NNEDenoiser/Source/NNEDenoiser/Private/NNEDenoiserModelInstanceGPU.h
+++ b/Engine/Plugins/Experimental/NNEDenoiser/Source/NNEDenoiser/Private/NNEDenoiserModelInstanceGPU.h
@@ -18,7 +18,7 @@ class FModelInstanceGPU : public IModelInstance
 public:
 	explicit FModelInstanceGPU(TSharedRef<UE::NNE::IModelInstanceGPU> ModelInstance);
 
-	static TUniquePtr<FModelInstanceGPU> Make(UNNEModelData& ModelData, const FString& RuntimeNameOverride = {});
+	static TUniquePtr<FModelInstanceGPU> Make(UNNEModelData& ModelData, const FString& RuntimeName);
 
 	virtual ~FModelInstanceGPU();
 
diff --git a/Engine/Plugins/Experimental/NNEDenoiser/Source/NNEDenoiser/Private/NNEDenoiserModelInstanceRDG.cpp b/Engine/Plugins/Experimental/NNEDenoiser/Source/NNEDenoiser/Private/NNEDenoiserModelInstanceRDG.cpp
--- a/Engine/Plugins/Experimental/NNEDenoiser/Source/NNEDenoiser/Private/NNEDenoiserModelInstanceRDG.cpp
+++ b/Engine/Plugins/Experimental/NNEDenoiser/Source/NNEDenoiser/Private/NNEDenoiserModelInstanceRDG.cpp
@@ -16,61 +16,43 @@
 namespace UE::NNEDenoiser::Private
 {
 
-	static const FString DefaultRuntimeRDGName = TEXT("NNERuntimeORTDml");
-	static const FString FallbackRuntimeRDGName = TEXT("NNERuntimeRDGHlsl");
-
-namespace Internal
-{
-	TUniquePtr<FModelInstanceRDG> Make(UNNEModelData& ModelData, const FString& RuntimeNameOverride)
+	TUniquePtr<FModelInstanceRDG> FModelInstanceRDG::Make(UNNEModelData& ModelData, const FString& RuntimeName)
 	{
-		FString RuntimeRDGName = RuntimeNameOverride.IsEmpty() ? DefaultRuntimeRDGName : RuntimeNameOverride;
+		check(!RuntimeName.IsEmpty());
 		
-		// Create the model
-		TWeakInterfacePtr<INNERuntimeRDG> RuntimeRDG = UE::NNE::GetRuntime<INNERuntimeRDG>(RuntimeRDGName);
-		if (!RuntimeRDG.IsValid())
+		TWeakInterfacePtr<INNERuntimeRDG> Runtime = UE::NNE::GetRuntime<INNERuntimeRDG>(RuntimeName);
+		if (!Runtime.IsValid())
 		{
-			UE_LOG(LogNNEDenoiser, Error, TEXT("No RDG runtime '%s' found. Valid RDG runtimes are: "), *RuntimeRDGName);
-			for (const FString& RuntimeName : UE::NNE::GetAllRuntimeNames<INNERuntimeRDG>())
+			UE_LOG(LogNNEDenoiser, Log, TEXT("Could not create model instance. No RDG runtime '%s' found. Valid RDG runtimes are: "), *RuntimeName);
+			for (const FString& Name : UE::NNE::GetAllRuntimeNames<INNERuntimeRDG>())
 			{
-				UE_LOG(LogNNEDenoiser, Error, TEXT("- %s"), *RuntimeName);
+				UE_LOG(LogNNEDenoiser, Log, TEXT("- %s"), *Name);
 			}
 			return {};
 		}
 
-		TSharedPtr<UE::NNE::IModelRDG> Model = RuntimeRDG->CreateModelRDG(ModelData);
+		if (Runtime->CanCreateModelRDG(ModelData) != INNERuntimeRDG::ECanCreateModelRDGStatus::Ok)
+		{
+			UE_LOG(LogNNEDenoiser, Log, TEXT("%s on RDG can not create model"), *RuntimeName);
+			return {};
+		}
+
+		TSharedPtr<UE::NNE::IModelRDG> Model = Runtime->CreateModelRDG(ModelData);
 		if (!Model.IsValid())
 		{
-			UE_LOG(LogNNEDenoiser, Error, TEXT("Could not create model using %s"), *RuntimeRDGName);
+			UE_LOG(LogNNEDenoiser, Log, TEXT("Could not create model using %s on RDG"), *RuntimeName);
 			return {};
 		}
 
 		TSharedPtr<UE::NNE::IModelInstanceRDG> ModelInstance = Model->CreateModelInstanceRDG();
 		if (!ModelInstance.IsValid())
 		{
-			UE_LOG(LogNNEDenoiser, Error, TEXT("Could not create model instance using %s"), *RuntimeRDGName);
+			UE_LOG(LogNNEDenoiser, Log, TEXT("Could not create model instance using %s on RDG"), *RuntimeName);
 			return {};
 		}
 
-		UE_LOG(LogNNEDenoiser, Log, TEXT("NNEDenoiserRDG: Creaded model instance from %s using %s"), *ModelData.GetFileId().ToString(), *RuntimeRDGName);
-
 		return MakeUnique<FModelInstanceRDG>(ModelInstance.ToSharedRef());
 	}
-} // namespace Internal
-
-	TUniquePtr<FModelInstanceRDG> FModelInstanceRDG::Make(UNNEModelData& ModelData, const FString& RuntimeNameOverride)
-	{
-		if (!RuntimeNameOverride.IsEmpty())
-		{
-			return Internal::Make(ModelData, RuntimeNameOverride);
-		}
-
-		if (TUniquePtr<FModelInstanceRDG> Result = Internal::Make(ModelData, DefaultRuntimeRDGName); Result.IsValid())
-		{
-			return Result;
-		}
-
-		return Internal::Make(ModelData, FallbackRuntimeRDGName);
-	}
 
 	FModelInstanceRDG::FModelInstanceRDG(TSharedRef<UE::NNE::IModelInstanceRDG> ModelInstance) :
 		ModelInstance(ModelInstance)
diff --git a/Engine/Plugins/Experimental/NNEDenoiser/Source/NNEDenoiser/Private/NNEDenoiserModelInstanceRDG.h b/Engine/Plugins/Experimental/NNEDenoiser/Source/NNEDenoiser/Private/NNEDenoiserModelInstanceRDG.h
--- a/Engine/Plugins/Experimental/NNEDenoiser/Source/NNEDenoiser/Private/NNEDenoiserModelInstanceRDG.h
+++ b/Engine/Plugins/Experimental/NNEDenoiser/Source/NNEDenoiser/Private/NNEDenoiserModelInstanceRDG.h
@@ -18,7 +18,7 @@ class FModelInstanceRDG : public IModelInstance
 public:
 	explicit FModelInstanceRDG(TSharedRef<UE::NNE::IModelInstanceRDG> ModelInstance);
 
-	static TUniquePtr<FModelInstanceRDG> Make(UNNEModelData& ModelData, const FString& RuntimeNameOverride = {});
+	static TUniquePtr<FModelInstanceRDG> Make(UNNEModelData& ModelData, const FString& RuntimeName);
 
 	virtual ~FModelInstanceRDG();
 
diff --git a/Engine/Plugins/Experimental/NNEDenoiser/Source/NNEDenoiser/Private/NNEDenoiserViewExtension.cpp b/Engine/Plugins/Experimental/NNEDenoiser/Source/NNEDenoiser/Private/NNEDenoiserViewExtension.cpp
--- a/Engine/Plugins/Experimental/NNEDenoiser/Source/NNEDenoiser/Private/NNEDenoiserViewExtension.cpp
+++ b/Engine/Plugins/Experimental/NNEDenoiser/Source/NNEDenoiser/Private/NNEDenoiserViewExtension.cpp
@@ -110,18 +110,65 @@ TUniquePtr<FGenericDenoiser> CreateNNEDenoiser(
 	TUniquePtr<IAutoExposure> AutoExposure,
 	TSharedPtr<ITransferFunction> TransferFunction)
 {
-	TUniquePtr<IModelInstance> ModelInstance;
-	if (RuntimeType == EDenoiserRuntimeType::CPU)
+	static TMap<EDenoiserRuntimeType, FString> DefaultRuntimeNames;
+	DefaultRuntimeNames.Emplace(EDenoiserRuntimeType::CPU, TEXT("NNERuntimeORTCpu"));
+	DefaultRuntimeNames.Emplace(EDenoiserRuntimeType::GPU, TEXT("NNERuntimeORTDml"));
+	DefaultRuntimeNames.Emplace(EDenoiserRuntimeType::RDG, TEXT("NNERuntimeORTDml"));
+
+	const static FString FallbackRuntimeNameRDG = TEXT("NNERuntimeRDGHlsl");
+
+	TArray<TPair<EDenoiserRuntimeType, FString>> RuntimePriorityQueue;
+	if (!RuntimeNameOverride.IsEmpty())
+	{
+		RuntimePriorityQueue.AddUnique({RuntimeType, RuntimeNameOverride});
+	}
+	
+	const FString& DefaultRuntimeName = DefaultRuntimeNames.FindChecked(RuntimeType);
+	if (RuntimeNameOverride != DefaultRuntimeName)
+	{
+		RuntimePriorityQueue.AddUnique({RuntimeType, DefaultRuntimeName});
+	}
+
+	if (RuntimeType >= EDenoiserRuntimeType::RDG)
 	{
-		ModelInstance = FModelInstanceCPU::Make(ModelData, RuntimeNameOverride);
+		if (RuntimeNameOverride != FallbackRuntimeNameRDG)
+		{
+			RuntimePriorityQueue.AddUnique({EDenoiserRuntimeType::RDG, FallbackRuntimeNameRDG});
+		}
+
+		const FString& DefaultRuntimeNameGPU = DefaultRuntimeNames.FindChecked(EDenoiserRuntimeType::GPU);
+		RuntimePriorityQueue.AddUnique({EDenoiserRuntimeType::GPU, DefaultRuntimeNameGPU});
 	}
-	else if (RuntimeType == EDenoiserRuntimeType::GPU)
+
+	if (RuntimeType >= EDenoiserRuntimeType::GPU)
 	{
-		ModelInstance = FModelInstanceGPU::Make(ModelData, RuntimeNameOverride);
+		const FString& DefaultRuntimeNameCPU = DefaultRuntimeNames.FindChecked(EDenoiserRuntimeType::CPU);
+		RuntimePriorityQueue.AddUnique({EDenoiserRuntimeType::CPU, DefaultRuntimeNameCPU});
 	}
-	else if (RuntimeType == EDenoiserRuntimeType::RDG)
+
+	TUniquePtr<IModelInstance> ModelInstance;
+	for (const auto& Pair : RuntimePriorityQueue)
 	{
-		ModelInstance = FModelInstanceRDG::Make(ModelData, RuntimeNameOverride);
+		UE_LOG(LogNNEDenoiser, Log, TEXT("Try create model instance with runtime %s on %s..."), *Pair.Value, *UEnum::GetValueAsString(Pair.Key));
+
+		if (Pair.Key == EDenoiserRuntimeType::CPU)
+		{
+			ModelInstance = FModelInstanceCPU::Make(ModelData, Pair.Value);
+		}
+		else if (Pair.Key == EDenoiserRuntimeType::GPU)
+		{
+			ModelInstance = FModelInstanceGPU::Make(ModelData, Pair.Value);
+		}
+		else if (Pair.Key == EDenoiserRuntimeType::RDG)
+		{
+			ModelInstance = FModelInstanceRDG::Make(ModelData, Pair.Value);
+		}
+
+		if (ModelInstance.IsValid())
+		{
+			UE_LOG(LogNNEDenoiser, Display, TEXT("Created model instance with runtime %s on %s"), *Pair.Value, *UEnum::GetValueAsString(Pair.Key));
+			break;
+		}
 	}
 
 	if (!ModelInstance.IsValid())
