[NNEDenoiser] Fix for Vulkan RHI. Corrected some resource transitions, updated texture/buffer readback and further cleanup.

#tests Editor, Win64 denoised on D3D12 and Vulkan RHIs.

#rnx

diff --git a/Engine/Plugins/Experimental/NNEDenoiser/Shaders/NNEDenoiserShadersAutoExposureReduce.usf b/Engine/Plugins/Experimental/NNEDenoiser/Shaders/NNEDenoiserShadersAutoExposureReduce.usf
--- a/Engine/Plugins/Experimental/NNEDenoiser/Shaders/NNEDenoiserShadersAutoExposureReduce.usf
+++ b/Engine/Plugins/Experimental/NNEDenoiser/Shaders/NNEDenoiserShadersAutoExposureReduce.usf
@@ -4,7 +4,7 @@
 #include "/Engine/Private/Common.ush"
 
 int InputSize;
-RWBuffer<float> InputBins;
+Buffer<float> InputBins;
 RWBuffer<float> OutputSums;
 RWBuffer<int> OutputCounts;
 int NumThreads;
diff --git a/Engine/Plugins/Experimental/NNEDenoiser/Shaders/NNEDenoiserShadersAutoExposureReduceFinal.usf b/Engine/Plugins/Experimental/NNEDenoiser/Shaders/NNEDenoiserShadersAutoExposureReduceFinal.usf
--- a/Engine/Plugins/Experimental/NNEDenoiser/Shaders/NNEDenoiserShadersAutoExposureReduceFinal.usf
+++ b/Engine/Plugins/Experimental/NNEDenoiser/Shaders/NNEDenoiserShadersAutoExposureReduceFinal.usf
@@ -4,8 +4,8 @@
 #include "/Engine/Private/Common.ush"
 
 int InputSize;
-RWBuffer<float> InputSums;
-RWBuffer<int> InputCounts;
+Buffer<float> InputSums;
+Buffer<int> InputCounts;
 RWBuffer<float> OutputBuffer; // 0: Result, 1: 1.0 / Result
 
 groupshared float Sums[THREAD_GROUP_SIZE];
diff --git a/Engine/Plugins/Experimental/NNEDenoiser/Shaders/NNEDenoiserShadersDefaultIOProcess.usf b/Engine/Plugins/Experimental/NNEDenoiser/Shaders/NNEDenoiserShadersDefaultIOProcess.usf
--- a/Engine/Plugins/Experimental/NNEDenoiser/Shaders/NNEDenoiserShadersDefaultIOProcess.usf
+++ b/Engine/Plugins/Experimental/NNEDenoiser/Shaders/NNEDenoiserShadersDefaultIOProcess.usf
@@ -3,12 +3,9 @@
 #include "/Engine/Public/Platform.ush"
 #include "/Engine/Private/Common.ush"
 
-int InputTextureWidth;
-int InputTextureHeight;
+int Width;
+int Height;
 Texture2D InputTexture;
-
-int OutputTextureWidth;
-int OutputTextureHeight;
 RWTexture2D<float4> OutputTexture;
 
 float3 GetProcessedValue(float3 Value)
@@ -32,7 +29,7 @@ float3 GetProcessedValue(float3 Value)
 [numthreads(THREAD_GROUP_SIZE, THREAD_GROUP_SIZE, 1)]
 void IOProcess(in const uint3 DispatchThreadID : SV_DispatchThreadID)
 {
-	if (DispatchThreadID.x >= OutputTextureWidth || DispatchThreadID.y >= OutputTextureHeight)
+	if (DispatchThreadID.x >= Width || DispatchThreadID.y >= Height)
 	{
 		return;
 	}
diff --git a/Engine/Plugins/Experimental/NNEDenoiser/Shaders/NNEDenoiserShadersMappedCopy.usf b/Engine/Plugins/Experimental/NNEDenoiser/Shaders/NNEDenoiserShadersMappedCopy.usf
--- a/Engine/Plugins/Experimental/NNEDenoiser/Shaders/NNEDenoiserShadersMappedCopy.usf
+++ b/Engine/Plugins/Experimental/NNEDenoiser/Shaders/NNEDenoiserShadersMappedCopy.usf
@@ -3,32 +3,56 @@
 #include "/Engine/Public/Platform.ush"
 #include "/Engine/Private/Common.ush"
 
-#if BUFFER_TYPE_INDEX == 4 // Float16
-	#define BUFFER_TYPE half
-	#define READ(X) float(X)
-	#define WRITE(X) half(X)
-#else
-	// BUFFER_TYPE_INDEX == 3 // Float32
-	#define BUFFER_TYPE float
-	#define READ(X) X
-	#define WRITE(X) X
+#define FLOAT16 half
+#define FLOAT16_4 half4
+
+#define FLOAT32 float
+#define FLOAT32_4 float4
+
+#if INPUT_DATA_TYPE_INDEX == 4 // Float16
+#define INPUT_SCALAR_DATA_TYPE FLOAT16
+#define INPUT_VECTOR_DATA_TYPE FLOAT16_4
+#else // INPUT_DATA_TYPE_INDEX == 3 // Float32
+#define INPUT_SCALAR_DATA_TYPE FLOAT32
+#define INPUT_VECTOR_DATA_TYPE FLOAT32_4
+#endif
+
+#if OUTPUT_DATA_TYPE_INDEX == 4 // Float16
+#define OUTPUT_SCALAR_DATA_TYPE FLOAT16
+#define OUTPUT_VECTOR_DATA_TYPE FLOAT16_4
+#else // OUTPUT_DATA_TYPE_INDEX == 3 // Float32
+#define OUTPUT_SCALAR_DATA_TYPE FLOAT32
+#define OUTPUT_VECTOR_DATA_TYPE FLOAT32_4
+#endif
+
+#if INPUT_DATA_TYPE_INDEX == 4 // Float16
+#if OUTPUT_DATA_TYPE_INDEX == 4 // Float16
+#define CONVERT(X) X
+#else // OUTPUT_DATA_TYPE_INDEX == 3 // Float32
+#define CONVERT(X) FLOAT32(X)
+#endif
+#else // INPUT_DATA_TYPE_INDEX == 3 // Float32
+#if OUTPUT_DATA_TYPE_INDEX == 4 // Float16
+#define CONVERT(X) FLOAT16(X)
+#else // OUTPUT_DATA_TYPE_INDEX == 3 // Float32
+#define CONVERT(X) X
+#endif
 #endif
 
 int Width;
 int Height;
-
 int4 OutputChannel_InputChannel_Unused_Unused[MAX_NUM_MAPPED_CHANNELS];
 
-#if INPUT_TYPE == 0 // Texture
-Texture2D InputTexture;
-#else // INPUT_TYPE == 1 // Buffer
-RWBuffer<BUFFER_TYPE> InputBuffer;
+#if INTRINSIC_INPUT_TYPE == 0 // Texture
+Texture2D<INPUT_VECTOR_DATA_TYPE> InputTexture;
+#else // INTRINSIC_INPUT_TYPE == 1 // Buffer
+Buffer<INPUT_SCALAR_DATA_TYPE> InputBuffer;
 #endif
 
-#if OUTPUT_TYPE == 0 // Texture
-RWTexture2D<float4> OutputTexture;
-#else // OUTPUT_TYPE == 1 // Buffer
-RWBuffer<BUFFER_TYPE> OutputBuffer;
+#if INTRINSIC_OUTPUT_TYPE == 0 // Texture
+RWTexture2D<OUTPUT_VECTOR_DATA_TYPE> OutputTexture;
+#else // INTRINSIC_OUTPUT_TYPE == 1 // Buffer
+RWBuffer<OUTPUT_SCALAR_DATA_TYPE> OutputBuffer;
 #endif
 
 [numthreads(THREAD_GROUP_SIZE, THREAD_GROUP_SIZE, 1)]
@@ -39,31 +63,40 @@ void MappedCopy(in const uint3 DispatchThreadID : SV_DispatchThreadID)
 		return;
 	}
 
-#if INPUT_TYPE == 1 || OUTPUT_TYPE == 1
+#if INTRINSIC_INPUT_TYPE == 1 || INTRINSIC_OUTPUT_TYPE == 1
 	const int Idx = DispatchThreadID.y * Width + DispatchThreadID.x;
 	const int Offset = Width * Height;
 #endif
 
-#if INPUT_TYPE == 0
-	const float4 InputColor = InputTexture.Load(int3(DispatchThreadID.x, DispatchThreadID.y, 0));
+#if INTRINSIC_INPUT_TYPE == 0
+	const INPUT_VECTOR_DATA_TYPE InputColor = InputTexture.Load(int3(DispatchThreadID.x, DispatchThreadID.y, 0));
 #endif
 
+#if INTRINSIC_OUTPUT_TYPE == 0
+	OUTPUT_VECTOR_DATA_TYPE Result = OutputTexture[DispatchThreadID.xy];
+#endif
+
+	UNROLL
 	for (int I = 0; I < NUM_MAPPED_CHANNELS; I++)
 	{
 		const int4 Info = OutputChannel_InputChannel_Unused_Unused[I];
 		const int OutputChannel = Info.x;
 		const int InputChannel = Info.y;
 
-#if INPUT_TYPE == 0
-		const float InputValue = InputColor[InputChannel];
+#if INTRINSIC_INPUT_TYPE == 0
+		const INPUT_SCALAR_DATA_TYPE InputValue = InputColor[InputChannel];
 #else
-		const float InputValue = InputBuffer[Idx + InputChannel * Offset];
+		const INPUT_SCALAR_DATA_TYPE InputValue = InputBuffer[Idx + InputChannel * Offset];
 #endif
 
-#if OUTPUT_TYPE == 0
-		OutputTexture[DispatchThreadID.xy][OutputChannel] = READ(InputValue);
+#if INTRINSIC_OUTPUT_TYPE == 0
+		Result[OutputChannel] = CONVERT(InputValue);
 #else
-		OutputBuffer[Idx + OutputChannel * Offset] = WRITE(InputValue);
+		OutputBuffer[Idx + OutputChannel * Offset] = CONVERT(InputValue);
 #endif
 	}
+
+#if INTRINSIC_OUTPUT_TYPE == 0
+	OutputTexture[DispatchThreadID.xy] = Result;
+#endif
 }
\ No newline at end of file
diff --git a/Engine/Plugins/Experimental/NNEDenoiser/Shaders/NNEDenoiserShadersTransferFunctionOidn.usf b/Engine/Plugins/Experimental/NNEDenoiser/Shaders/NNEDenoiserShadersTransferFunctionOidn.usf
--- a/Engine/Plugins/Experimental/NNEDenoiser/Shaders/NNEDenoiserShadersTransferFunctionOidn.usf
+++ b/Engine/Plugins/Experimental/NNEDenoiser/Shaders/NNEDenoiserShadersTransferFunctionOidn.usf
@@ -3,18 +3,14 @@
 #include "/Engine/Public/Platform.ush"
 #include "/Engine/Private/Common.ush"
 
-int InputTextureWidth;
-int InputTextureHeight;
+int Width;
+int Height;
 Texture2D InputTexture;
-
-int OutputTextureWidth;
-int OutputTextureHeight;
-RWTexture2D<float4> OutputTexture;
-
-RWBuffer<float> InputScaleBuffer;
-
+Buffer<float> InputScaleBuffer;
 float NormScale;
 float InvNormScale;
+RWTexture2D<float4> OutputTexture;
+
 
 ///
 
@@ -69,7 +65,7 @@ float3 GetProcessedValue(float3 Value)
 [numthreads(THREAD_GROUP_SIZE, THREAD_GROUP_SIZE, 1)]
 void PreOrPostprocess(in const uint3 DispatchThreadID : SV_DispatchThreadID)
 {
-	if (DispatchThreadID.x >= OutputTextureWidth || DispatchThreadID.y >= OutputTextureHeight)
+	if (DispatchThreadID.x >= Width || DispatchThreadID.y >= Height)
 	{
 		return;
 	}
diff --git a/Engine/Plugins/Experimental/NNEDenoiser/Source/NNEDenoiser/Private/NNEDenoiserAutoExposure.cpp b/Engine/Plugins/Experimental/NNEDenoiser/Source/NNEDenoiser/Private/NNEDenoiserAutoExposure.cpp
--- a/Engine/Plugins/Experimental/NNEDenoiser/Source/NNEDenoiser/Private/NNEDenoiserAutoExposure.cpp
+++ b/Engine/Plugins/Experimental/NNEDenoiser/Source/NNEDenoiser/Private/NNEDenoiserAutoExposure.cpp
@@ -130,7 +130,7 @@ void FAutoExposure::EnqueueRDG(FRDGBuilder& GraphBuilder, FRDGTextureRef InputTe
 	{
 		FAutoExposureReduceCS::FParameters *ShaderParameters = GraphBuilder.AllocParameters<FAutoExposureReduceCS::FParameters>();
 		ShaderParameters->InputSize = NumBins;
-		ShaderParameters->InputBins = GraphBuilder.CreateUAV(FRDGBufferUAVDesc(OutputBins, EPixelFormat::PF_R32_FLOAT));
+		ShaderParameters->InputBins = GraphBuilder.CreateSRV(FRDGBufferSRVDesc(OutputBins, EPixelFormat::PF_R32_FLOAT));
 		ShaderParameters->OutputSums = GraphBuilder.CreateUAV(FRDGBufferUAVDesc(OutputSums, EPixelFormat::PF_R32_FLOAT));
 		ShaderParameters->OutputCounts = GraphBuilder.CreateUAV(FRDGBufferUAVDesc(OutputCounts, EPixelFormat::PF_R32_SINT));
 		ShaderParameters->NumThreads = NumReduceGroups * FAutoExposureDownsampleConstants::MAX_BIN_SIZE;
@@ -153,8 +153,8 @@ void FAutoExposure::EnqueueRDG(FRDGBuilder& GraphBuilder, FRDGTextureRef InputTe
 	{
 		FAutoExposureReduceFinalCS::FParameters *ShaderParameters = GraphBuilder.AllocParameters<FAutoExposureReduceFinalCS::FParameters>();
 		ShaderParameters->InputSize = NumReduceGroups;
-		ShaderParameters->InputSums = GraphBuilder.CreateUAV(FRDGBufferUAVDesc(OutputSums, EPixelFormat::PF_R32_FLOAT));
-		ShaderParameters->InputCounts = GraphBuilder.CreateUAV(FRDGBufferUAVDesc(OutputCounts, EPixelFormat::PF_R32_SINT));
+		ShaderParameters->InputSums = GraphBuilder.CreateSRV(FRDGBufferSRVDesc(OutputSums, EPixelFormat::PF_R32_FLOAT));
+		ShaderParameters->InputCounts = GraphBuilder.CreateSRV(FRDGBufferSRVDesc(OutputCounts, EPixelFormat::PF_R32_SINT));
 		ShaderParameters->OutputBuffer = GraphBuilder.CreateUAV(FRDGBufferUAVDesc(OutputBuffer, EPixelFormat::PF_R32_FLOAT));
 
 		FGlobalShaderMap* GlobalShaderMap = GetGlobalShaderMap(GMaxRHIFeatureLevel);
diff --git a/Engine/Plugins/Experimental/NNEDenoiser/Source/NNEDenoiser/Private/NNEDenoiserIOProcessBase.cpp b/Engine/Plugins/Experimental/NNEDenoiser/Source/NNEDenoiser/Private/NNEDenoiserIOProcessBase.cpp
--- a/Engine/Plugins/Experimental/NNEDenoiser/Source/NNEDenoiser/Private/NNEDenoiserIOProcessBase.cpp
+++ b/Engine/Plugins/Experimental/NNEDenoiser/Source/NNEDenoiser/Private/NNEDenoiserIOProcessBase.cpp
@@ -48,23 +48,21 @@ void AddPreOrPostProcess(
 {
 	using namespace UE::NNEDenoiserShaders::Internal;
 
-	const FIntVector InputTextureSize = InputTexture->Desc.GetSize();
-	const FIntVector OutputTextureSize = OutputTexture->Desc.GetSize();
+	const FIntVector Size = InputTexture->Desc.GetSize();
+	check(Size == OutputTexture->Desc.GetSize());
 
 	FDefaultIOProcessCS::FParameters *ShaderParameters = GraphBuilder.AllocParameters<FDefaultIOProcessCS::FParameters>();
-	ShaderParameters->InputTextureWidth = InputTextureSize.X;
-	ShaderParameters->InputTextureHeight = InputTextureSize.Y;
+	ShaderParameters->Width = Size.X;
+	ShaderParameters->Height = Size.Y;
 	ShaderParameters->InputTexture = InputTexture;
-	ShaderParameters->OutputTextureWidth = OutputTextureSize.X;
-	ShaderParameters->OutputTextureHeight = OutputTextureSize.Y;
 	ShaderParameters->OutputTexture = GraphBuilder.CreateUAV(FRDGTextureUAVDesc(OutputTexture));
 
 	FDefaultIOProcessCS::FPermutationDomain PermutationVector;
 	PermutationVector.Set<FDefaultIOProcessCS::FDefaultIOProcessInputKind>(GetInputKind(TensorName));
 
 	FIntVector ThreadGroupCount = FIntVector(
-		FMath::DivideAndRoundUp(OutputTextureSize.X, FDefaultIOProcessConstants::THREAD_GROUP_SIZE),
-		FMath::DivideAndRoundUp(OutputTextureSize.Y, FDefaultIOProcessConstants::THREAD_GROUP_SIZE),
+		FMath::DivideAndRoundUp(Size.X, FDefaultIOProcessConstants::THREAD_GROUP_SIZE),
+		FMath::DivideAndRoundUp(Size.Y, FDefaultIOProcessConstants::THREAD_GROUP_SIZE),
 		1);
 
 	FGlobalShaderMap* GlobalShaderMap = GetGlobalShaderMap(GMaxRHIFeatureLevel);
@@ -85,9 +83,7 @@ void AddPreOrPostProcess(
 void AddReadInputPass(
 	FRDGBuilder& GraphBuilder,
 	FRDGTextureRef InputTexture,
-	FRDGBufferUAVRef InputBufferUAV,
-	FIntPoint InputBufferSize,
-	UE::NNEDenoiserShaders::Internal::ENNEDenoiserDataType DataType,
+	FRDGBufferUAVRef OutputBufferUAV,
 	const TArray<FIntPoint>& ChannelMapping)
 {
 	using namespace UE::NNEDenoiserShaders::Internal;
@@ -104,29 +100,33 @@ void AddReadInputPass(
 		return;
 	}
 
-	const FIntVector InputTextureSize = InputTexture->Desc.GetSize();
+	const FIntVector Size = InputTexture->Desc.GetSize();
 
-	FNNEDenoiserTextureBufferMappedCopyCS::FParameters *ReadInputParameters = GraphBuilder.AllocParameters<FNNEDenoiserTextureBufferMappedCopyCS::FParameters>();
-	ReadInputParameters->Width = InputBufferSize.X;
-	ReadInputParameters->Height = InputBufferSize.Y;
+	FTextureBufferMappedCopyCS::FParameters *ReadInputParameters = GraphBuilder.AllocParameters<FTextureBufferMappedCopyCS::FParameters>();
+	ReadInputParameters->Width = Size.X;
+	ReadInputParameters->Height = Size.Y;
 	ReadInputParameters->InputTexture = InputTexture;
-	ReadInputParameters->OutputBuffer = InputBufferUAV;
+	ReadInputParameters->OutputBuffer = OutputBufferUAV;
 	for (int32 Idx = 0; Idx < ChannelMapping.Num(); Idx++)
 	{
 		ReadInputParameters->OutputChannel_InputChannel_Unused_Unused[Idx] = { ChannelMapping[Idx].X, ChannelMapping[Idx].Y, 0, 0 };
 	}
 
-	FNNEDenoiserTextureBufferMappedCopyCS::FPermutationDomain PermutationVector;
-	PermutationVector.Set<FNNEDenoiserTextureBufferMappedCopyCS::FNNEDenoiserDataType>(DataType);
-	PermutationVector.Set<FNNEDenoiserTextureBufferMappedCopyCS::FNNEDenoiserNumMappedChannels>(ChannelMapping.Num());
+	const EDataType InputDataType = GetDenoiserShaderDataType(InputTexture->Desc.Format);
+	const EDataType OutputDataType = GetDenoiserShaderDataType(OutputBufferUAV->Desc.Format);
+
+	FTextureBufferMappedCopyCS::FPermutationDomain PermutationVector;
+	PermutationVector.Set<FTextureBufferMappedCopyCS::FInputDataType>(InputDataType);
+	PermutationVector.Set<FTextureBufferMappedCopyCS::FOutputDataType>(OutputDataType);
+	PermutationVector.Set<FTextureBufferMappedCopyCS::FNumMappedChannels>(ChannelMapping.Num());
 
 	FIntVector ReadInputThreadGroupCount = FIntVector(
-		FMath::DivideAndRoundUp(InputBufferSize.X, FMappedCopyConstants::THREAD_GROUP_SIZE),
-		FMath::DivideAndRoundUp(InputBufferSize.Y, FMappedCopyConstants::THREAD_GROUP_SIZE),
+		FMath::DivideAndRoundUp(Size.X, FMappedCopyConstants::THREAD_GROUP_SIZE),
+		FMath::DivideAndRoundUp(Size.Y, FMappedCopyConstants::THREAD_GROUP_SIZE),
 		1);
 
 	FGlobalShaderMap* GlobalShaderMap = GetGlobalShaderMap(GMaxRHIFeatureLevel);
-	TShaderMapRef<FNNEDenoiserTextureBufferMappedCopyCS> ReadInputShader(GlobalShaderMap, PermutationVector);
+	TShaderMapRef<FTextureBufferMappedCopyCS> ReadInputShader(GlobalShaderMap, PermutationVector);
 
 	RDG_EVENT_SCOPE(GraphBuilder, "NNEDenoiser.ReadInput");
 	RDG_GPU_STAT_SCOPE(GraphBuilder, FNNEDenoiserReadInput);
@@ -142,37 +142,41 @@ void AddReadInputPass(
 
 void AddWriteOutputPass(
 	FRDGBuilder& GraphBuilder,
-	FRDGBufferUAVRef OutputBufferUAV,
-	FIntPoint OutputBufferSize,
-	FRDGTextureUAVRef OutputTextureUAV,
-	UE::NNEDenoiserShaders::Internal::ENNEDenoiserDataType DataType,
+	FRDGBufferRef InputBuffer,
+	EPixelFormat BufferFormat,
+	FRDGTextureRef OutputTexture,
+	UE::NNEDenoiserShaders::Internal::EDataType DataType,
 	const TArray<FIntPoint>& ChannelMapping)
 {
 	using namespace UE::NNEDenoiserShaders::Internal;
 
-	const FIntVector OutputTextureSize = OutputTextureUAV->Desc.Texture->Desc.GetSize();
+	const FIntVector Size = OutputTexture->Desc.GetSize();
 
-	FNNEDenoiserBufferTextureMappedCopyCS::FParameters *WriteOutputParameters = GraphBuilder.AllocParameters<FNNEDenoiserBufferTextureMappedCopyCS::FParameters>();
-	WriteOutputParameters->Width = OutputBufferSize.X;
-	WriteOutputParameters->Height = OutputBufferSize.Y;
-	WriteOutputParameters->InputBuffer = OutputBufferUAV;
-	WriteOutputParameters->OutputTexture = OutputTextureUAV;
+	FBufferTextureMappedCopyCS::FParameters *WriteOutputParameters = GraphBuilder.AllocParameters<FBufferTextureMappedCopyCS::FParameters>();
+	WriteOutputParameters->Width = Size.X;
+	WriteOutputParameters->Height = Size.Y;
+	WriteOutputParameters->InputBuffer = GraphBuilder.CreateSRV(FRDGBufferSRVDesc(InputBuffer, BufferFormat));
+	WriteOutputParameters->OutputTexture = GraphBuilder.CreateUAV(FRDGTextureUAVDesc(OutputTexture));
 	for (int32 Idx = 0; Idx < ChannelMapping.Num(); Idx++)
 	{
 		WriteOutputParameters->OutputChannel_InputChannel_Unused_Unused[Idx] = { ChannelMapping[Idx].Y, ChannelMapping[Idx].X, 0, 0 };
 	}
 
-	FNNEDenoiserBufferTextureMappedCopyCS::FPermutationDomain PermutationVector;
-	PermutationVector.Set<FNNEDenoiserBufferTextureMappedCopyCS::FNNEDenoiserDataType>(DataType);
-	PermutationVector.Set<FNNEDenoiserBufferTextureMappedCopyCS::FNNEDenoiserNumMappedChannels>(ChannelMapping.Num());
+	const EDataType InputDataType = GetDenoiserShaderDataType(BufferFormat);
+	const EDataType OutputDataType = GetDenoiserShaderDataType(OutputTexture->Desc.Format);
+
+	FBufferTextureMappedCopyCS::FPermutationDomain PermutationVector;
+	PermutationVector.Set<FBufferTextureMappedCopyCS::FInputDataType>(InputDataType);
+	PermutationVector.Set<FBufferTextureMappedCopyCS::FOutputDataType>(OutputDataType);
+	PermutationVector.Set<FBufferTextureMappedCopyCS::FNumMappedChannels>(ChannelMapping.Num());
 
 	FIntVector WriteOutputThreadGroupCount = FIntVector(
-		FMath::DivideAndRoundUp(OutputTextureSize.X, FMappedCopyConstants::THREAD_GROUP_SIZE),
-		FMath::DivideAndRoundUp(OutputTextureSize.Y, FMappedCopyConstants::THREAD_GROUP_SIZE),
+		FMath::DivideAndRoundUp(Size.X, FMappedCopyConstants::THREAD_GROUP_SIZE),
+		FMath::DivideAndRoundUp(Size.Y, FMappedCopyConstants::THREAD_GROUP_SIZE),
 		1);
 
 	FGlobalShaderMap* GlobalShaderMap = GetGlobalShaderMap(GMaxRHIFeatureLevel);
-	TShaderMapRef<FNNEDenoiserBufferTextureMappedCopyCS> WriteOutputShader(GlobalShaderMap, PermutationVector);
+	TShaderMapRef<FBufferTextureMappedCopyCS> WriteOutputShader(GlobalShaderMap, PermutationVector);
 
 	RDG_EVENT_SCOPE(GraphBuilder, "NNEDenoiser.WriteOutput");
 	RDG_GPU_STAT_SCOPE(GraphBuilder, FNNEDenoiserWriteOutput);
@@ -189,8 +193,6 @@ void AddWriteOutputPass(
 void AddReadInputPassForKind(
 	FRDGBuilder& GraphBuilder,
 	const IResourceAccess& ResourceAccess,
-	FIntPoint BufferSize,
-	NNEDenoiserShaders::Internal::ENNEDenoiserDataType DataType,
 	EResourceName TensorName,
 	const FResourceMapping& ResourceMapping,
 	FRDGBufferUAVRef BufferUAV)
@@ -200,7 +202,7 @@ void AddReadInputPassForKind(
 		const int32 FrameIdx = -ChannelMappingKeyValue.Key;
 		FRDGTextureRef InputTexture = ResourceAccess.GetTexture(TensorName, FrameIdx);
 
-		AddReadInputPass(GraphBuilder, InputTexture, BufferUAV, BufferSize, DataType, ChannelMappingKeyValue.Value);
+		AddReadInputPass(GraphBuilder, InputTexture, BufferUAV, ChannelMappingKeyValue.Value);
 	}
 };
 
@@ -381,16 +383,14 @@ void FInputProcessBase::WriteInputBuffer(
 	FRDGBufferRef Buffer) const
 {
 	const ENNETensorDataType TensorDataType = TensorDesc.GetDataType();
-	const NNEDenoiserShaders::Internal::ENNEDenoiserDataType DataType = GetDenoiserShaderDataType(TensorDataType);
-	const FIntPoint BufferSize(TensorShape.GetData()[3], TensorShape.GetData()[2]);
 
 	FRDGBufferUAVRef BufferUAV = GraphBuilder.CreateUAV(FRDGBufferUAVDesc(Buffer, GetBufferFormat(TensorDataType)));
 
-	IOProcessBaseHelper::AddReadInputPassForKind(GraphBuilder, ResourceAccess, BufferSize, DataType, EResourceName::Color, ResourceMapping, BufferUAV);
-	IOProcessBaseHelper::AddReadInputPassForKind(GraphBuilder, ResourceAccess, BufferSize, DataType, EResourceName::Albedo, ResourceMapping, BufferUAV);
-	IOProcessBaseHelper::AddReadInputPassForKind(GraphBuilder, ResourceAccess, BufferSize, DataType, EResourceName::Normal, ResourceMapping, BufferUAV);
-	IOProcessBaseHelper::AddReadInputPassForKind(GraphBuilder, ResourceAccess, BufferSize, DataType, EResourceName::Flow, ResourceMapping, BufferUAV);
-	IOProcessBaseHelper::AddReadInputPassForKind(GraphBuilder, ResourceAccess, BufferSize, DataType, EResourceName::Output, ResourceMapping, BufferUAV);
+	IOProcessBaseHelper::AddReadInputPassForKind(GraphBuilder, ResourceAccess, EResourceName::Color, ResourceMapping, BufferUAV);
+	IOProcessBaseHelper::AddReadInputPassForKind(GraphBuilder, ResourceAccess, EResourceName::Albedo, ResourceMapping, BufferUAV);
+	IOProcessBaseHelper::AddReadInputPassForKind(GraphBuilder, ResourceAccess, EResourceName::Normal, ResourceMapping, BufferUAV);
+	IOProcessBaseHelper::AddReadInputPassForKind(GraphBuilder, ResourceAccess, EResourceName::Flow, ResourceMapping, BufferUAV);
+	IOProcessBaseHelper::AddReadInputPassForKind(GraphBuilder, ResourceAccess, EResourceName::Output, ResourceMapping, BufferUAV);
 }
 
 bool FOutputProcessBase::Validate(const IModelInstance& ModelInstance, FIntPoint Extent) const
@@ -495,17 +495,16 @@ void FOutputProcessBase::ReadOutputBuffer(
 	FRDGTextureRef Texture) const
 {
 	const ENNETensorDataType TensorDataType = TensorDesc.GetDataType();
-	const NNEDenoiserShaders::Internal::ENNEDenoiserDataType DataType = GetDenoiserShaderDataType(TensorDataType);
+	const NNEDenoiserShaders::Internal::EDataType DataType = GetDenoiserShaderDataType(TensorDataType);
 	const FIntPoint BufferSize(TensorShape.GetData()[3], TensorShape.GetData()[2]);
 		
 	TMap<int32, TArray<FIntPoint>> ChannelMappingPerFrame = ResourceMapping.GetChannelMappingPerFrame(EResourceName::Output);
 	checkf(ChannelMappingPerFrame.Num() == 1, TEXT("Invalid output mapping"));
 	const TArray<FIntPoint>& ChannelMapping = ChannelMappingPerFrame.CreateConstIterator().Value();
 
-	FRDGBufferUAVRef BufferUAV = GraphBuilder.CreateUAV(FRDGBufferUAVDesc(Buffer, GetBufferFormat(TensorDataType)));
-	FRDGTextureUAVRef TextureUAV = GraphBuilder.CreateUAV(FRDGTextureUAVDesc(Texture));
+	const EPixelFormat BufferFormat = GetBufferFormat(TensorDataType);
 
-	IOProcessBaseHelper::AddWriteOutputPass(GraphBuilder, BufferUAV, BufferSize, TextureUAV, DataType, ChannelMapping);
+	IOProcessBaseHelper::AddWriteOutputPass(GraphBuilder, Buffer, BufferFormat, Texture, DataType, ChannelMapping);
 }
 
 void FOutputProcessBase::PostprocessOutput(
diff --git a/Engine/Plugins/Experimental/NNEDenoiser/Source/NNEDenoiser/Private/NNEDenoiserModelInstanceCPU.cpp b/Engine/Plugins/Experimental/NNEDenoiser/Source/NNEDenoiser/Private/NNEDenoiserModelInstanceCPU.cpp
--- a/Engine/Plugins/Experimental/NNEDenoiser/Source/NNEDenoiser/Private/NNEDenoiserModelInstanceCPU.cpp
+++ b/Engine/Plugins/Experimental/NNEDenoiser/Source/NNEDenoiser/Private/NNEDenoiserModelInstanceCPU.cpp
@@ -105,7 +105,7 @@ namespace UE::NNEDenoiser::Private
 		}
 		for (const NNE::FTensorBindingRDG& Binding : Outputs)
 		{
-			DenoiserParameters->OutputBuffers.Emplace(Binding.Buffer, ERHIAccess::CopySrc);
+			DenoiserParameters->OutputBuffers.Emplace(Binding.Buffer, ERHIAccess::CopyDest);
 		}
 		
 		GraphBuilder.AddPass(RDG_EVENT_NAME("NNEDenoiser.DenoiseCPU"), DenoiserParameters, ERDGPassFlags::Readback,
diff --git a/Engine/Plugins/Experimental/NNEDenoiser/Source/NNEDenoiser/Private/NNEDenoiserModelInstanceGPU.cpp b/Engine/Plugins/Experimental/NNEDenoiser/Source/NNEDenoiser/Private/NNEDenoiserModelInstanceGPU.cpp
--- a/Engine/Plugins/Experimental/NNEDenoiser/Source/NNEDenoiser/Private/NNEDenoiserModelInstanceGPU.cpp
+++ b/Engine/Plugins/Experimental/NNEDenoiser/Source/NNEDenoiser/Private/NNEDenoiserModelInstanceGPU.cpp
@@ -108,7 +108,7 @@ namespace UE::NNEDenoiser::Private
 		}
 		for (const NNE::FTensorBindingRDG& Binding : Outputs)
 		{
-			DenoiserParameters->OutputBuffers.Emplace(Binding.Buffer, ERHIAccess::CopySrc);
+			DenoiserParameters->OutputBuffers.Emplace(Binding.Buffer, ERHIAccess::CopyDest);
 		}
 
 		RDG_EVENT_SCOPE(GraphBuilder, "NNEDenoiser.ModelInstanceGPU");
diff --git a/Engine/Plugins/Experimental/NNEDenoiser/Source/NNEDenoiser/Private/NNEDenoiserResourceManager.cpp b/Engine/Plugins/Experimental/NNEDenoiser/Source/NNEDenoiser/Private/NNEDenoiserResourceManager.cpp
--- a/Engine/Plugins/Experimental/NNEDenoiser/Source/NNEDenoiser/Private/NNEDenoiserResourceManager.cpp
+++ b/Engine/Plugins/Experimental/NNEDenoiser/Source/NNEDenoiser/Private/NNEDenoiserResourceManager.cpp
@@ -52,6 +52,10 @@ void FResourceManager::AddTexture(EResourceName Name, FRDGTextureRef Current, in
 	// note: Always create intermediate texture for now.
 	const bool bNeedIntermediateTexture = true;
 
+#if UE_BUILD_DEBUG
+	static TSet<FString> DebugNames;
+#endif
+
 	if (bNeedIntermediateTexture)
 	{
 		check(!IntermediateTextureMap.Contains(Name));
@@ -63,7 +67,23 @@ void FResourceManager::AddTexture(EResourceName Name, FRDGTextureRef Current, in
 
 		for (int32 I = 0; I < IntermediateTextures.Num(); I++)
 		{
+#if UE_BUILD_DEBUG
+			const FString ResourceName = UEnum::GetValueAsName<EResourceName>(Name).ToString();
+			const FString FrameNumberText = FString::FromInt(I);
+			const FString TextureDebugName = FString(TEXT("NNEDenoiser.IntermediateTexture_")) + ResourceName + FrameNumberText;
+
+			const TCHAR* DebugNamePtr = nullptr;
+			if (!DebugNames.Contains(TextureDebugName))
+			{
+				DebugNames.Add(TextureDebugName);
+			}
+
+			DebugNamePtr = **DebugNames.Find(TextureDebugName);
+
+			IntermediateTextures[I] = GraphBuilder.CreateTexture(Desc, DebugNamePtr);
+#else
 			IntermediateTextures[I] = GraphBuilder.CreateTexture(Desc, TEXT("NNEDenoiser.IntermediateTexture"));
+#endif
 		}
 	}
 }
diff --git a/Engine/Plugins/Experimental/NNEDenoiser/Source/NNEDenoiser/Private/NNEDenoiserTransferFunctionOidn.cpp b/Engine/Plugins/Experimental/NNEDenoiser/Source/NNEDenoiser/Private/NNEDenoiserTransferFunctionOidn.cpp
--- a/Engine/Plugins/Experimental/NNEDenoiser/Source/NNEDenoiser/Private/NNEDenoiserTransferFunctionOidn.cpp
+++ b/Engine/Plugins/Experimental/NNEDenoiser/Source/NNEDenoiser/Private/NNEDenoiserTransferFunctionOidn.cpp
@@ -61,26 +61,24 @@ void ApplyTransferFunction(
 {
 	using namespace UE::NNEDenoiserShaders::Internal;
 
-	const FIntVector InputTextureSize = InputTexture->Desc.GetSize();
-	const FIntVector OutputTextureSize = OutputTexture->Desc.GetSize();
+	const FIntVector Size = InputTexture->Desc.GetSize();
+	check(Size == OutputTexture->Desc.GetSize());
 
 	FTransferFunctionOidnCS::FParameters *ShaderParameters = GraphBuilder.AllocParameters<FTransferFunctionOidnCS::FParameters>();
-	ShaderParameters->InputTextureWidth = InputTextureSize.X;
-	ShaderParameters->InputTextureHeight = InputTextureSize.Y;
+	ShaderParameters->Width = Size.X;
+	ShaderParameters->Height = Size.Y;
 	ShaderParameters->InputTexture = InputTexture;
-	ShaderParameters->OutputTextureWidth = OutputTextureSize.X;
-	ShaderParameters->OutputTextureHeight = OutputTextureSize.Y;
-	ShaderParameters->OutputTexture = GraphBuilder.CreateUAV(FRDGTextureUAVDesc(OutputTexture));
-	ShaderParameters->InputScaleBuffer = GraphBuilder.CreateUAV(FRDGBufferUAVDesc(InputScaleBuffer, EPixelFormat::PF_R32_FLOAT));
+	ShaderParameters->InputScaleBuffer = GraphBuilder.CreateSRV(FRDGBufferSRVDesc(InputScaleBuffer, EPixelFormat::PF_R32_FLOAT));
 	ShaderParameters->NormScale = NormScale;
 	ShaderParameters->InvNormScale = InvNormScale;
+	ShaderParameters->OutputTexture = GraphBuilder.CreateUAV(FRDGTextureUAVDesc(OutputTexture));
 
 	FTransferFunctionOidnCS::FPermutationDomain PermutationVector;
 	PermutationVector.Set<FTransferFunctionOidnCS::FTransferFunctionOidnMode>(Forward ? ETransferFunctionOidnMode::Forwward : ETransferFunctionOidnMode::Inverse);
 
 	FIntVector ThreadGroupCount = FIntVector(
-		FMath::DivideAndRoundUp(OutputTextureSize.X, FTransferFunctionOidnConstants::THREAD_GROUP_SIZE),
-		FMath::DivideAndRoundUp(OutputTextureSize.Y, FTransferFunctionOidnConstants::THREAD_GROUP_SIZE),
+		FMath::DivideAndRoundUp(Size.X, FTransferFunctionOidnConstants::THREAD_GROUP_SIZE),
+		FMath::DivideAndRoundUp(Size.Y, FTransferFunctionOidnConstants::THREAD_GROUP_SIZE),
 		1);
 
 	FGlobalShaderMap* GlobalShaderMap = GetGlobalShaderMap(GMaxRHIFeatureLevel);
diff --git a/Engine/Plugins/Experimental/NNEDenoiser/Source/NNEDenoiser/Private/NNEDenoiserUtils.h b/Engine/Plugins/Experimental/NNEDenoiser/Source/NNEDenoiser/Private/NNEDenoiserUtils.h
--- a/Engine/Plugins/Experimental/NNEDenoiser/Source/NNEDenoiser/Private/NNEDenoiserUtils.h
+++ b/Engine/Plugins/Experimental/NNEDenoiser/Source/NNEDenoiser/Private/NNEDenoiserUtils.h
@@ -8,6 +8,7 @@
 #include "NNEDenoiserShadersMappedCopyCS.h"
 #include "NNETypes.h"
 #include "RHICommandList.h"
+#include "RHIGPUReadback.h"
 #include "RHIResources.h"
 #include "RHITypes.h"
 
@@ -55,16 +56,21 @@ namespace UE::NNEDenoiser::Private
 	{
 		SCOPED_NAMED_EVENT_TEXT("NNEDenoiser.CopyTextureFromGPUToCPU", FColor::Magenta);
 
-		uint32 SrcStride = 0;
-		const PixelType* SrcBuffer = static_cast<PixelType*>(RHICmdList.LockTexture2D(SrcTexture, 0, RLM_ReadOnly, SrcStride, false));
-		SrcStride /= sizeof(PixelType);
+		FRHIGPUTextureReadback Readback(TEXT("NNEDenoiser.CopyTextureFromGPUToCPU"));
+		Readback.EnqueueCopy(RHICmdList, SrcTexture, FIntVector::ZeroValue, 0, FIntVector(Size.X, Size.Y, 1));
+
+		RHICmdList.BlockUntilGPUIdle();
+
+		int32_t SrcStride = 0;
+		const PixelType* SrcBuffer = static_cast<PixelType*>(Readback.Lock(SrcStride, nullptr));
+
 		PixelType* DstBuffer = DstArray.GetData();
-		for (int32 Y = 0; Y < Size.Y; Y++, DstBuffer += Size.X, SrcBuffer += SrcStride)
+		for (int Y = 0; Y < Size.Y; Y++, DstBuffer += Size.X, SrcBuffer += SrcStride)
 		{
 			FPlatformMemory::Memcpy(DstBuffer, SrcBuffer, Size.X * sizeof(PixelType));
 
 		}
-		RHICmdList.UnlockTexture2D(SrcTexture, 0, false);
+		Readback.Unlock();
 	}
 
 	template <typename PixelType>
@@ -88,12 +94,19 @@ namespace UE::NNEDenoiser::Private
 	{
 		SCOPED_NAMED_EVENT_TEXT("NNEDenoiser.CopyBufferFromGPUToCPU", FColor::Magenta);
 
-		const ElementType* Src = static_cast<ElementType*>(RHICmdList.LockBuffer(SrcBuffer, 0, Count * sizeof(ElementType), RLM_ReadOnly));
+		const int32 NumBytes = Count * sizeof(ElementType);
+
+		FRHIGPUBufferReadback Readback(TEXT("NNEDenoiser.CopyBufferFromGPUToCPU"));
+		Readback.EnqueueCopy(RHICmdList, SrcBuffer, NumBytes);
+
+		RHICmdList.BlockUntilGPUIdle();
+
+		const ElementType* Src = static_cast<ElementType*>(Readback.Lock(NumBytes));
 		ElementType* Dst = DstArray.GetData();
 		
-		FPlatformMemory::Memcpy(Dst, Src, Count * sizeof(ElementType));
+		FPlatformMemory::Memcpy(Dst, Src, NumBytes);
 
-		RHICmdList.UnlockBuffer(SrcBuffer);
+		Readback.Unlock();
 	}
 
 	template <typename ElementType>
@@ -115,20 +128,42 @@ namespace UE::NNEDenoiser::Private
 		{
 			case ENNETensorDataType::Half: return EPixelFormat::PF_R16F;
 			case ENNETensorDataType::Float: return EPixelFormat::PF_R32_FLOAT;
+			default:
+				unimplemented();
 		}
 		return EPixelFormat::PF_Unknown;
 	}
 
-	inline UE::NNEDenoiserShaders::Internal::ENNEDenoiserDataType GetDenoiserShaderDataType(ENNETensorDataType TensorDataType)
+	inline UE::NNEDenoiserShaders::Internal::EDataType GetDenoiserShaderDataType(ENNETensorDataType TensorDataType)
 	{
-		using UE::NNEDenoiserShaders::Internal::ENNEDenoiserDataType;
+		using UE::NNEDenoiserShaders::Internal::EDataType;
 
 		switch (TensorDataType)
 		{
-			case ENNETensorDataType::Half: return ENNEDenoiserDataType::Half;
-			case ENNETensorDataType::Float: return ENNEDenoiserDataType::Float;
+			case ENNETensorDataType::Half: return EDataType::Half;
+			case ENNETensorDataType::Float: return EDataType::Float;
+			default:
+				unimplemented();
 		}
-		return ENNEDenoiserDataType::None;
+		return EDataType::None;
 	}
 
+	inline UE::NNEDenoiserShaders::Internal::EDataType GetDenoiserShaderDataType(EPixelFormat Format)
+	{
+		using UE::NNEDenoiserShaders::Internal::EDataType;
+
+		switch (Format)
+		{
+			case EPixelFormat::PF_R16F:
+			case EPixelFormat::PF_FloatRGBA: return EDataType::Half;
+
+			case EPixelFormat::PF_R32_FLOAT:
+			case EPixelFormat::PF_A32B32G32R32F: return EDataType::Float;
+			
+			default:
+				unimplemented();
+		}
+		return EDataType::None;
+	} 
+
 }
\ No newline at end of file
diff --git a/Engine/Plugins/Experimental/NNEDenoiser/Source/NNEDenoiserShaders/Internal/NNEDenoiserShadersAutoExposureCS.h b/Engine/Plugins/Experimental/NNEDenoiser/Source/NNEDenoiserShaders/Internal/NNEDenoiserShadersAutoExposureCS.h
--- a/Engine/Plugins/Experimental/NNEDenoiser/Source/NNEDenoiserShaders/Internal/NNEDenoiserShadersAutoExposureCS.h
+++ b/Engine/Plugins/Experimental/NNEDenoiser/Source/NNEDenoiserShaders/Internal/NNEDenoiserShadersAutoExposureCS.h
@@ -9,24 +9,6 @@
 namespace UE::NNEDenoiserShaders::Internal
 {
 	
-	class NNEDENOISERSHADERS_API FAutoExposureCS : public FGlobalShader
-	{
-		DECLARE_GLOBAL_SHADER(FAutoExposureCS);
-		SHADER_USE_PARAMETER_STRUCT(FAutoExposureCS, FGlobalShader)
-
-	public:
-
-		BEGIN_SHADER_PARAMETER_STRUCT(FParameters, )
-			SHADER_PARAMETER(int32, InputTextureWidth)
-			SHADER_PARAMETER(int32, InputTextureHeight)
-			SHADER_PARAMETER_RDG_TEXTURE(Texture2D, InputTexture)
-			SHADER_PARAMETER_RDG_BUFFER_UAV(RWBuffer, OutputBuffer)
-		END_SHADER_PARAMETER_STRUCT()
-
-		static void ModifyCompilationEnvironment(const FGlobalShaderPermutationParameters& InParameters, FShaderCompilerEnvironment& OutEnvironment);
-		static bool ShouldCompilePermutation(const FGlobalShaderPermutationParameters& Parameters);
-	};
-
 	class FAutoExposureDownsampleConstants
 	{
 	public:
@@ -71,7 +53,7 @@ namespace UE::NNEDenoiserShaders::Internal
 
 		BEGIN_SHADER_PARAMETER_STRUCT(FParameters, )
 			SHADER_PARAMETER(int32, InputSize)
-			SHADER_PARAMETER_RDG_BUFFER_UAV(RWBuffer, InputBins)
+			SHADER_PARAMETER_RDG_BUFFER_SRV(Buffer, InputBins)
 			SHADER_PARAMETER_RDG_BUFFER_UAV(RWBuffer, OutputSums)
 			SHADER_PARAMETER_RDG_BUFFER_UAV(RWBuffer, OutputCounts)
 			SHADER_PARAMETER(int32, NumThreads)
@@ -90,8 +72,8 @@ namespace UE::NNEDenoiserShaders::Internal
 
 		BEGIN_SHADER_PARAMETER_STRUCT(FParameters, )
 			SHADER_PARAMETER(int32, InputSize)
-			SHADER_PARAMETER_RDG_BUFFER_UAV(RWBuffer, InputSums)
-			SHADER_PARAMETER_RDG_BUFFER_UAV(RWBuffer, InputCounts)
+			SHADER_PARAMETER_RDG_BUFFER_SRV(Buffer, InputSums)
+			SHADER_PARAMETER_RDG_BUFFER_SRV(Buffer, InputCounts)
 			SHADER_PARAMETER_RDG_BUFFER_UAV(RWBuffer, OutputBuffer)
 		END_SHADER_PARAMETER_STRUCT()
 
diff --git a/Engine/Plugins/Experimental/NNEDenoiser/Source/NNEDenoiserShaders/Internal/NNEDenoiserShadersDefaultIOProcessCS.h b/Engine/Plugins/Experimental/NNEDenoiser/Source/NNEDenoiserShaders/Internal/NNEDenoiserShadersDefaultIOProcessCS.h
--- a/Engine/Plugins/Experimental/NNEDenoiser/Source/NNEDenoiserShaders/Internal/NNEDenoiserShadersDefaultIOProcessCS.h
+++ b/Engine/Plugins/Experimental/NNEDenoiser/Source/NNEDenoiserShaders/Internal/NNEDenoiserShadersDefaultIOProcessCS.h
@@ -36,11 +36,9 @@ namespace UE::NNEDenoiserShaders::Internal
 	public:
 
 		BEGIN_SHADER_PARAMETER_STRUCT(FParameters, )
-			SHADER_PARAMETER(int32, InputTextureWidth)
-			SHADER_PARAMETER(int32, InputTextureHeight)
+			SHADER_PARAMETER(int32, Width)
+			SHADER_PARAMETER(int32, Height)
 			SHADER_PARAMETER_RDG_TEXTURE(Texture2D, InputTexture)
-			SHADER_PARAMETER(int32, OutputTextureWidth)
-			SHADER_PARAMETER(int32, OutputTextureHeight)
 			SHADER_PARAMETER_RDG_TEXTURE_UAV(RWTexture2D, OutputTexture)
 		END_SHADER_PARAMETER_STRUCT()
 
diff --git a/Engine/Plugins/Experimental/NNEDenoiser/Source/NNEDenoiserShaders/Internal/NNEDenoiserShadersMappedCopyCS.h b/Engine/Plugins/Experimental/NNEDenoiser/Source/NNEDenoiserShaders/Internal/NNEDenoiserShadersMappedCopyCS.h
--- a/Engine/Plugins/Experimental/NNEDenoiser/Source/NNEDenoiserShaders/Internal/NNEDenoiserShadersMappedCopyCS.h
+++ b/Engine/Plugins/Experimental/NNEDenoiser/Source/NNEDenoiserShaders/Internal/NNEDenoiserShadersMappedCopyCS.h
@@ -9,7 +9,8 @@
 namespace UE::NNEDenoiserShaders::Internal
 {
 
-	enum class ENNEDenoiserDataType : uint8
+	// Maps to ENNETensorDataType
+	enum class EDataType : uint8
 	{
 		None = 0,
 		// Char,
@@ -25,15 +26,18 @@ namespace UE::NNEDenoiserShaders::Internal
 		static constexpr int32 THREAD_GROUP_SIZE{16};
 		static constexpr int32 MAX_NUM_MAPPED_CHANNELS{4};
 	};
+
+	// Note: The shader also supports Texture/Texture and Buffer/Buffer mapped copy, just add the Shader class.
 	
-	class NNEDENOISERSHADERS_API FNNEDenoiserTextureBufferMappedCopyCS : public FGlobalShader
+	class NNEDENOISERSHADERS_API FTextureBufferMappedCopyCS : public FGlobalShader
 	{
-		DECLARE_GLOBAL_SHADER(FNNEDenoiserTextureBufferMappedCopyCS);
-		SHADER_USE_PARAMETER_STRUCT(FNNEDenoiserTextureBufferMappedCopyCS, FGlobalShader)
+		DECLARE_GLOBAL_SHADER(FTextureBufferMappedCopyCS);
+		SHADER_USE_PARAMETER_STRUCT(FTextureBufferMappedCopyCS, FGlobalShader)
 
-		class FNNEDenoiserDataType : SHADER_PERMUTATION_ENUM_CLASS("BUFFER_TYPE_INDEX", ENNEDenoiserDataType);
-		class FNNEDenoiserNumMappedChannels : SHADER_PERMUTATION_RANGE_INT("NUM_MAPPED_CHANNELS", 0, FMappedCopyConstants::MAX_NUM_MAPPED_CHANNELS + 1);
-		using FPermutationDomain = TShaderPermutationDomain<FNNEDenoiserDataType, FNNEDenoiserNumMappedChannels>;
+		class FInputDataType : SHADER_PERMUTATION_ENUM_CLASS("INPUT_DATA_TYPE_INDEX", EDataType);
+		class FOutputDataType : SHADER_PERMUTATION_ENUM_CLASS("OUTPUT_DATA_TYPE_INDEX", EDataType);
+		class FNumMappedChannels : SHADER_PERMUTATION_RANGE_INT("NUM_MAPPED_CHANNELS", 0, FMappedCopyConstants::MAX_NUM_MAPPED_CHANNELS + 1);
+		using FPermutationDomain = TShaderPermutationDomain<FInputDataType, FOutputDataType, FNumMappedChannels>;
 
 	public:
 
@@ -49,67 +53,22 @@ namespace UE::NNEDenoiserShaders::Internal
 		static bool ShouldCompilePermutation(const FGlobalShaderPermutationParameters& Parameters);
 	};
 
-	class NNEDENOISERSHADERS_API FNNEDenoiserBufferMappedCopyCS : public FGlobalShader
-	{
-		DECLARE_GLOBAL_SHADER(FNNEDenoiserBufferMappedCopyCS);
-		SHADER_USE_PARAMETER_STRUCT(FNNEDenoiserBufferMappedCopyCS, FGlobalShader)
-
-		class FNNEDenoiserDataType : SHADER_PERMUTATION_ENUM_CLASS("BUFFER_TYPE_INDEX", ENNEDenoiserDataType);
-		class FNNEDenoiserNumMappedChannels : SHADER_PERMUTATION_RANGE_INT("NUM_MAPPED_CHANNELS", 0, FMappedCopyConstants::MAX_NUM_MAPPED_CHANNELS + 1);
-		using FPermutationDomain = TShaderPermutationDomain<FNNEDenoiserDataType, FNNEDenoiserNumMappedChannels>;
-
-	public:
-
-		BEGIN_SHADER_PARAMETER_STRUCT(FParameters, )
-			SHADER_PARAMETER(int32, Width)
-			SHADER_PARAMETER(int32, Height)
-			SHADER_PARAMETER_RDG_BUFFER_UAV(RWBuffer, InputBuffer)
-			SHADER_PARAMETER_RDG_BUFFER_UAV(RWBuffer, OutputBuffer)
-			SHADER_PARAMETER_ARRAY(FIntVector4, OutputChannel_InputChannel_Unused_Unused, [FMappedCopyConstants::MAX_NUM_MAPPED_CHANNELS])
-		END_SHADER_PARAMETER_STRUCT()
-
-		static void ModifyCompilationEnvironment(const FGlobalShaderPermutationParameters& InParameters, FShaderCompilerEnvironment& OutEnvironment);
-		static bool ShouldCompilePermutation(const FGlobalShaderPermutationParameters& Parameters);
-	};
-
-	class NNEDENOISERSHADERS_API FNNEDenoiserBufferTextureMappedCopyCS : public FGlobalShader
-	{
-		DECLARE_GLOBAL_SHADER(FNNEDenoiserBufferTextureMappedCopyCS);
-		SHADER_USE_PARAMETER_STRUCT(FNNEDenoiserBufferTextureMappedCopyCS, FGlobalShader)
-
-		class FNNEDenoiserDataType : SHADER_PERMUTATION_ENUM_CLASS("BUFFER_TYPE_INDEX", ENNEDenoiserDataType);
-		class FNNEDenoiserNumMappedChannels : SHADER_PERMUTATION_RANGE_INT("NUM_MAPPED_CHANNELS", 0, FMappedCopyConstants::MAX_NUM_MAPPED_CHANNELS + 1);
-		using FPermutationDomain = TShaderPermutationDomain<FNNEDenoiserDataType, FNNEDenoiserNumMappedChannels>;
-
-	public:
-
-		BEGIN_SHADER_PARAMETER_STRUCT(FParameters, )
-			SHADER_PARAMETER(int32, Width)
-			SHADER_PARAMETER(int32, Height)
-			SHADER_PARAMETER_RDG_BUFFER_UAV(RWBuffer, InputBuffer)
-			SHADER_PARAMETER_RDG_TEXTURE_UAV(RWTexture2D, OutputTexture)
-			SHADER_PARAMETER_ARRAY(FIntVector4, OutputChannel_InputChannel_Unused_Unused, [FMappedCopyConstants::MAX_NUM_MAPPED_CHANNELS])
-		END_SHADER_PARAMETER_STRUCT()
-
-		static void ModifyCompilationEnvironment(const FGlobalShaderPermutationParameters& InParameters, FShaderCompilerEnvironment& OutEnvironment);
-		static bool ShouldCompilePermutation(const FGlobalShaderPermutationParameters& Parameters);
-	};
-
-	class NNEDENOISERSHADERS_API FNNEDenoiserTextureMappedCopyCS : public FGlobalShader
+	class NNEDENOISERSHADERS_API FBufferTextureMappedCopyCS : public FGlobalShader
 	{
-		DECLARE_GLOBAL_SHADER(FNNEDenoiserTextureMappedCopyCS);
-		SHADER_USE_PARAMETER_STRUCT(FNNEDenoiserTextureMappedCopyCS, FGlobalShader)
+		DECLARE_GLOBAL_SHADER(FBufferTextureMappedCopyCS);
+		SHADER_USE_PARAMETER_STRUCT(FBufferTextureMappedCopyCS, FGlobalShader)
 
-		class FNNEDenoiserDataType : SHADER_PERMUTATION_ENUM_CLASS("BUFFER_TYPE_INDEX", ENNEDenoiserDataType);
-		class FNNEDenoiserNumMappedChannels : SHADER_PERMUTATION_RANGE_INT("NUM_MAPPED_CHANNELS", 0, FMappedCopyConstants::MAX_NUM_MAPPED_CHANNELS + 1);
-		using FPermutationDomain = TShaderPermutationDomain<FNNEDenoiserDataType, FNNEDenoiserNumMappedChannels>;
+		class FInputDataType : SHADER_PERMUTATION_ENUM_CLASS("INPUT_DATA_TYPE_INDEX", EDataType);
+		class FOutputDataType : SHADER_PERMUTATION_ENUM_CLASS("OUTPUT_DATA_TYPE_INDEX", EDataType);
+		class FNumMappedChannels : SHADER_PERMUTATION_RANGE_INT("NUM_MAPPED_CHANNELS", 0, FMappedCopyConstants::MAX_NUM_MAPPED_CHANNELS + 1);
+		using FPermutationDomain = TShaderPermutationDomain<FInputDataType, FOutputDataType, FNumMappedChannels>;
 
 	public:
 
 		BEGIN_SHADER_PARAMETER_STRUCT(FParameters, )
 			SHADER_PARAMETER(int32, Width)
 			SHADER_PARAMETER(int32, Height)
-			SHADER_PARAMETER_RDG_TEXTURE_UAV(RWTexture2D, InputTexture)
+			SHADER_PARAMETER_RDG_BUFFER_SRV(Buffer, InputBuffer)
 			SHADER_PARAMETER_RDG_TEXTURE_UAV(RWTexture2D, OutputTexture)
 			SHADER_PARAMETER_ARRAY(FIntVector4, OutputChannel_InputChannel_Unused_Unused, [FMappedCopyConstants::MAX_NUM_MAPPED_CHANNELS])
 		END_SHADER_PARAMETER_STRUCT()
diff --git a/Engine/Plugins/Experimental/NNEDenoiser/Source/NNEDenoiserShaders/Internal/NNEDenoiserShadersTransferFunctionOidnCS.h b/Engine/Plugins/Experimental/NNEDenoiser/Source/NNEDenoiserShaders/Internal/NNEDenoiserShadersTransferFunctionOidnCS.h
--- a/Engine/Plugins/Experimental/NNEDenoiser/Source/NNEDenoiserShaders/Internal/NNEDenoiserShadersTransferFunctionOidnCS.h
+++ b/Engine/Plugins/Experimental/NNEDenoiser/Source/NNEDenoiserShaders/Internal/NNEDenoiserShadersTransferFunctionOidnCS.h
@@ -33,15 +33,13 @@ namespace UE::NNEDenoiserShaders::Internal
 	public:
 
 		BEGIN_SHADER_PARAMETER_STRUCT(FParameters, )
-			SHADER_PARAMETER(int32, InputTextureWidth)
-			SHADER_PARAMETER(int32, InputTextureHeight)
+			SHADER_PARAMETER(int32, Width)
+			SHADER_PARAMETER(int32, Height)
 			SHADER_PARAMETER_RDG_TEXTURE(Texture2D, InputTexture)
-			SHADER_PARAMETER(int32, OutputTextureWidth)
-			SHADER_PARAMETER(int32, OutputTextureHeight)
-			SHADER_PARAMETER_RDG_TEXTURE_UAV(RWTexture2D, OutputTexture)
-			SHADER_PARAMETER_RDG_BUFFER_UAV(RWBuffer, InputScaleBuffer)
+			SHADER_PARAMETER_RDG_BUFFER_SRV(Buffer, InputScaleBuffer)
 			SHADER_PARAMETER(float, NormScale)
 			SHADER_PARAMETER(float, InvNormScale)
+			SHADER_PARAMETER_RDG_TEXTURE_UAV(RWTexture2D, OutputTexture)
 		END_SHADER_PARAMETER_STRUCT()
 
 		static void ModifyCompilationEnvironment(const FGlobalShaderPermutationParameters& InParameters, FShaderCompilerEnvironment& OutEnvironment);
diff --git a/Engine/Plugins/Experimental/NNEDenoiser/Source/NNEDenoiserShaders/Private/NNEDenoiserShadersAutoExposureCS.cpp b/Engine/Plugins/Experimental/NNEDenoiser/Source/NNEDenoiserShaders/Private/NNEDenoiserShadersAutoExposureCS.cpp
--- a/Engine/Plugins/Experimental/NNEDenoiser/Source/NNEDenoiserShaders/Private/NNEDenoiserShadersAutoExposureCS.cpp
+++ b/Engine/Plugins/Experimental/NNEDenoiser/Source/NNEDenoiserShaders/Private/NNEDenoiserShadersAutoExposureCS.cpp
@@ -4,11 +4,6 @@
 
 namespace UE::NNEDenoiserShaders::Internal
 {
-	void FAutoExposureCS::ModifyCompilationEnvironment(const FGlobalShaderPermutationParameters& InParameters, FShaderCompilerEnvironment& OutEnvironment)
-	{
-		FGlobalShader::ModifyCompilationEnvironment(InParameters, OutEnvironment);
-		OutEnvironment.SetDefine(TEXT("MAX_FLT"), MAX_FLT);
-	}
 
 	void FAutoExposureDownsampleCS::ModifyCompilationEnvironment(const FGlobalShaderPermutationParameters& InParameters, FShaderCompilerEnvironment& OutEnvironment)
 	{
@@ -31,11 +26,6 @@ namespace UE::NNEDenoiserShaders::Internal
 		OutEnvironment.SetDefine(TEXT("KEY"), FAutoExposureReduceConstants::KEY);
 	}
 
-	bool FAutoExposureCS::ShouldCompilePermutation(const FGlobalShaderPermutationParameters& Parameters)
-	{
-		return ShouldCompileRayTracingShadersForProject(Parameters.Platform);
-	}
-
 	bool FAutoExposureDownsampleCS::ShouldCompilePermutation(const FGlobalShaderPermutationParameters& Parameters)
 	{
 		return ShouldCompileRayTracingShadersForProject(Parameters.Platform);
@@ -51,7 +41,6 @@ namespace UE::NNEDenoiserShaders::Internal
 		return ShouldCompileRayTracingShadersForProject(Parameters.Platform);
 	}
 
-	IMPLEMENT_GLOBAL_SHADER(FAutoExposureCS, "/NNEDenoiserShaders/NNEDenoiserShadersAutoExposure.usf", "AutoExposure", SF_Compute);
 	IMPLEMENT_GLOBAL_SHADER(FAutoExposureDownsampleCS, "/NNEDenoiserShaders/NNEDenoiserShadersAutoExposureDownsample.usf", "Downsample", SF_Compute);
 	IMPLEMENT_GLOBAL_SHADER(FAutoExposureReduceCS, "/NNEDenoiserShaders/NNEDenoiserShadersAutoExposureReduce.usf", "Reduce", SF_Compute);
 	IMPLEMENT_GLOBAL_SHADER(FAutoExposureReduceFinalCS, "/NNEDenoiserShaders/NNEDenoiserShadersAutoExposureReduceFinal.usf", "ReduceFinal", SF_Compute);
diff --git a/Engine/Plugins/Experimental/NNEDenoiser/Source/NNEDenoiserShaders/Private/NNEDenoiserShadersMappedCopyCS.cpp b/Engine/Plugins/Experimental/NNEDenoiser/Source/NNEDenoiserShaders/Private/NNEDenoiserShadersMappedCopyCS.cpp
--- a/Engine/Plugins/Experimental/NNEDenoiser/Source/NNEDenoiserShaders/Private/NNEDenoiserShadersMappedCopyCS.cpp
+++ b/Engine/Plugins/Experimental/NNEDenoiser/Source/NNEDenoiserShaders/Private/NNEDenoiserShadersMappedCopyCS.cpp
@@ -11,58 +11,54 @@ namespace UE::NNEDenoiserShaders::Internal
 		OutEnvironment.SetDefine(TEXT("MAX_NUM_MAPPED_CHANNELS"), FMappedCopyConstants::MAX_NUM_MAPPED_CHANNELS);
 	}
 
-	void FNNEDenoiserTextureBufferMappedCopyCS::ModifyCompilationEnvironment(const FGlobalShaderPermutationParameters& InParameters, FShaderCompilerEnvironment& OutEnvironment)
+	template<class GlobalShaderType>
+	bool CommonShouldCompilePermutation(const FGlobalShaderPermutationParameters& Parameters)
 	{
-		CommonModifyCompilationEnvironment(InParameters, OutEnvironment);
-		OutEnvironment.SetDefine(TEXT("INPUT_TYPE"), 0);
-		OutEnvironment.SetDefine(TEXT("OUTPUT_TYPE"), 1);
-	}
+		if (!ShouldCompileRayTracingShadersForProject(Parameters.Platform))
+		{
+			return false;
+		}
 
-	void FNNEDenoiserBufferMappedCopyCS::ModifyCompilationEnvironment(const FGlobalShaderPermutationParameters& InParameters, FShaderCompilerEnvironment& OutEnvironment)
-	{
-		CommonModifyCompilationEnvironment(InParameters, OutEnvironment);
-		OutEnvironment.SetDefine(TEXT("INPUT_TYPE"), 1);
-		OutEnvironment.SetDefine(TEXT("OUTPUT_TYPE"), 1);
-	}
+		typename GlobalShaderType::FPermutationDomain PermutationVector(Parameters.PermutationId);
+		if (PermutationVector.template Get<typename GlobalShaderType::FInputDataType>() != EDataType::Half &&
+			PermutationVector.template Get<typename GlobalShaderType::FInputDataType>() != EDataType::Float)
+		{
+			return false;
+		}
+		if (PermutationVector.template Get<typename GlobalShaderType::FOutputDataType>() != EDataType::Half &&
+			PermutationVector.template Get<typename GlobalShaderType::FOutputDataType>() != EDataType::Float)
+		{
+			return false;
+		}
 
-	void FNNEDenoiserBufferTextureMappedCopyCS::ModifyCompilationEnvironment(const FGlobalShaderPermutationParameters& InParameters, FShaderCompilerEnvironment& OutEnvironment)
-	{
-		CommonModifyCompilationEnvironment(InParameters, OutEnvironment);
-		OutEnvironment.SetDefine(TEXT("INPUT_TYPE"), 1);
-		OutEnvironment.SetDefine(TEXT("OUTPUT_TYPE"), 0);
+		return true;
 	}
 
-	void FNNEDenoiserTextureMappedCopyCS::ModifyCompilationEnvironment(const FGlobalShaderPermutationParameters& InParameters, FShaderCompilerEnvironment& OutEnvironment)
+	void FTextureBufferMappedCopyCS::ModifyCompilationEnvironment(const FGlobalShaderPermutationParameters& InParameters, FShaderCompilerEnvironment& OutEnvironment)
 	{
 		CommonModifyCompilationEnvironment(InParameters, OutEnvironment);
-		OutEnvironment.SetDefine(TEXT("INPUT_TYPE"), 0);
-		OutEnvironment.SetDefine(TEXT("OUTPUT_TYPE"), 0);
-	}
-
-	bool FNNEDenoiserTextureBufferMappedCopyCS::ShouldCompilePermutation(const FGlobalShaderPermutationParameters& Parameters)
-	{
-		return ShouldCompileRayTracingShadersForProject(Parameters.Platform);
+		OutEnvironment.SetDefine(TEXT("INTRINSIC_INPUT_TYPE"), 0);
+		OutEnvironment.SetDefine(TEXT("INTRINSIC_OUTPUT_TYPE"), 1);
 	}
 
-	bool FNNEDenoiserBufferMappedCopyCS::ShouldCompilePermutation(const FGlobalShaderPermutationParameters& Parameters)
+	void FBufferTextureMappedCopyCS::ModifyCompilationEnvironment(const FGlobalShaderPermutationParameters& InParameters, FShaderCompilerEnvironment& OutEnvironment)
 	{
-		return ShouldCompileRayTracingShadersForProject(Parameters.Platform);
+		CommonModifyCompilationEnvironment(InParameters, OutEnvironment);
+		OutEnvironment.SetDefine(TEXT("INTRINSIC_INPUT_TYPE"), 1);
+		OutEnvironment.SetDefine(TEXT("INTRINSIC_OUTPUT_TYPE"), 0);
 	}
 
-	bool FNNEDenoiserBufferTextureMappedCopyCS::ShouldCompilePermutation(const FGlobalShaderPermutationParameters& Parameters)
+	bool FTextureBufferMappedCopyCS::ShouldCompilePermutation(const FGlobalShaderPermutationParameters& Parameters)
 	{
-		return ShouldCompileRayTracingShadersForProject(Parameters.Platform);
+		return CommonShouldCompilePermutation<FTextureBufferMappedCopyCS>(Parameters);
 	}
 
-	bool FNNEDenoiserTextureMappedCopyCS::ShouldCompilePermutation(const FGlobalShaderPermutationParameters& Parameters)
+	bool FBufferTextureMappedCopyCS::ShouldCompilePermutation(const FGlobalShaderPermutationParameters& Parameters)
 	{
-		return ShouldCompileRayTracingShadersForProject(Parameters.Platform);
+		return CommonShouldCompilePermutation<FBufferTextureMappedCopyCS>(Parameters);
 	}
 
-	IMPLEMENT_GLOBAL_SHADER(FNNEDenoiserTextureBufferMappedCopyCS, "/NNEDenoiserShaders/NNEDenoiserShadersMappedCopy.usf", "MappedCopy", SF_Compute);
-	IMPLEMENT_GLOBAL_SHADER(FNNEDenoiserBufferMappedCopyCS, "/NNEDenoiserShaders/NNEDenoiserShadersMappedCopy.usf", "MappedCopy", SF_Compute);
-
-	IMPLEMENT_GLOBAL_SHADER(FNNEDenoiserBufferTextureMappedCopyCS, "/NNEDenoiserShaders/NNEDenoiserShadersMappedCopy.usf", "MappedCopy", SF_Compute);
-	IMPLEMENT_GLOBAL_SHADER(FNNEDenoiserTextureMappedCopyCS, "/NNEDenoiserShaders/NNEDenoiserShadersMappedCopy.usf", "MappedCopy", SF_Compute);
+	IMPLEMENT_GLOBAL_SHADER(FTextureBufferMappedCopyCS, "/NNEDenoiserShaders/NNEDenoiserShadersMappedCopy.usf", "MappedCopy", SF_Compute);
+	IMPLEMENT_GLOBAL_SHADER(FBufferTextureMappedCopyCS, "/NNEDenoiserShaders/NNEDenoiserShadersMappedCopy.usf", "MappedCopy", SF_Compute);
 
 } // UE::NNEDenoiser::Private
\ No newline at end of file
diff --git a/Engine/Plugins/Experimental/NNEDenoiser/Source/NNEDenoiserShaders/Private/NNEDenoiserShadersTransferFunctionOidnCS.cpp b/Engine/Plugins/Experimental/NNEDenoiser/Source/NNEDenoiserShaders/Private/NNEDenoiserShadersTransferFunctionOidnCS.cpp
new file mode 100644
--- /dev/null
+++ b/Engine/Plugins/Experimental/NNEDenoiser/Source/NNEDenoiserShaders/Private/NNEDenoiserShadersTransferFunctionOidnCS.cpp
@@ -0,0 +1,21 @@
+// Copyright Epic Games, Inc. All Rights Reserved.
+
+#include "NNEDenoiserShadersTransferFunctionOidnCS.h"
+
+namespace UE::NNEDenoiserShaders::Internal
+{
+	void FTransferFunctionOidnCS::ModifyCompilationEnvironment(const FGlobalShaderPermutationParameters& InParameters, FShaderCompilerEnvironment& OutEnvironment)
+	{
+		FGlobalShader::ModifyCompilationEnvironment(InParameters, OutEnvironment);
+		OutEnvironment.SetDefine(TEXT("THREAD_GROUP_SIZE"), FTransferFunctionOidnConstants::THREAD_GROUP_SIZE);
+		OutEnvironment.SetDefine(TEXT("MAX_FLT"), MAX_FLT);
+	}
+
+	bool FTransferFunctionOidnCS::ShouldCompilePermutation(const FGlobalShaderPermutationParameters& Parameters)
+	{
+		return ShouldCompileRayTracingShadersForProject(Parameters.Platform);
+	}
+
+	IMPLEMENT_GLOBAL_SHADER(FTransferFunctionOidnCS, "/NNEDenoiserShaders/NNEDenoiserShadersTransferFunctionOidn.usf", "PreOrPostprocess", SF_Compute);
+
+} // UE::NNEDenoiser::Private
\ No newline at end of file
diff --git a/Engine/Plugins/Experimental/NNEDenoiser/Source/NNEDenoiserShaders/Private/NNEDenoiserTransferFunctionOidn.cpp b/Engine/Plugins/Experimental/NNEDenoiser/Source/NNEDenoiserShaders/Private/NNEDenoiserTransferFunctionOidn.cpp
deleted file mode 100644
