[NNEDenoiser] Upgrade OIDN models to version 2.3.0.

#tests Editor, Win64 compared screenshots to OIDN and perf measurements

#virtualized

diff --git a/Engine/Plugins/Experimental/NNEDenoiser/Source/NNEDenoiser/Private/NNEDenoiserSettings.cpp b/Engine/Plugins/Experimental/NNEDenoiser/Source/NNEDenoiser/Private/NNEDenoiserSettings.cpp
--- a/Engine/Plugins/Experimental/NNEDenoiser/Source/NNEDenoiser/Private/NNEDenoiserSettings.cpp
+++ b/Engine/Plugins/Experimental/NNEDenoiser/Source/NNEDenoiser/Private/NNEDenoiserSettings.cpp
@@ -11,7 +11,7 @@ UNNEDenoiserSettings::UNNEDenoiserSettings()
 	CategoryName = TEXT("Plugins");
 	SectionName = TEXT("NNE Denoiser");
 
-	DenoiserModelData = FSoftObjectPath(TEXT("/NNEDenoiser/NNEDNN_Oidn2-2-0_ColorAlbedoNormal_Dyn_Float16.NNEDNN_Oidn2-2-0_ColorAlbedoNormal_Dyn_Float16"));
+	DenoiserModelData = FSoftObjectPath(TEXT("/NNEDenoiser/NNEDNN_Oidn2-3-0_ColorAlbedoNormal.NNEDNN_Oidn2-3-0_ColorAlbedoNormal"));
 }
 
 void UNNEDenoiserSettings::PostInitProperties()
diff --git a/Engine/Plugins/Experimental/NNEDenoiser/Source/NNEDenoiser/Private/NNEDenoiserViewExtension.cpp b/Engine/Plugins/Experimental/NNEDenoiser/Source/NNEDenoiser/Private/NNEDenoiserViewExtension.cpp
--- a/Engine/Plugins/Experimental/NNEDenoiser/Source/NNEDenoiser/Private/NNEDenoiserViewExtension.cpp
+++ b/Engine/Plugins/Experimental/NNEDenoiser/Source/NNEDenoiser/Private/NNEDenoiserViewExtension.cpp
@@ -32,18 +32,12 @@ static TAutoConsoleVariable<int32> CVarNNEDenoiserModelData(
 	0,
 	TEXT("Defines the model data of the denoiser model.\n")
 	TEXT("  0: Use model data from Project Settings\n")
-	TEXT("  1: OIDN Color\n")
-	TEXT("  2: OIDN Color and features\n")
-	TEXT("  3: OIDN High Quality (prefilters features)\n")
-	TEXT("  4: OIDN Color | Float16\n")
-	TEXT("  5: OIDN Color and features | Float16\n")
-	TEXT("  6: OIDN High Quality (prefilters features) | Float16\n")
-	TEXT("  7: OIDN Color | Alpha\n")
-	TEXT("  8: OIDN Color and features | Alpha\n")
-	TEXT("  9: OIDN High Quality (prefilters features) | Alpha\n")
-	TEXT("  10: OIDN Color | Alpha | Float16\n")
-	TEXT("  11: OIDN Color and features | Alpha | Float16\n")
-	TEXT("  12: OIDN High Quality (prefilters features) | Alpha | Float16")
+	TEXT("  1: OIDN Fast\n")
+	TEXT("  2: OIDN Balanced\n")
+	TEXT("  3: OIDN High Quality\n")
+	TEXT("  4: OIDN Fast | Alpha\n")
+	TEXT("  5: OIDN Balanced | Alpha\n")
+	TEXT("  6: OIDN High Quality | Alpha")
 );
 
 static TAutoConsoleVariable<int32> CVarNNEDenoiserRuntimeType(
@@ -76,25 +70,14 @@ FString GetDenoiserModelDataNameFromCVarAndSettings(const UNNEDenoiserSettings*
 	{
 		case 0: return !Settings->DenoiserModelData.IsNull() ? Settings->DenoiserModelData.ToString() : FString();
 
-		// Float32
-		case 1: return TEXT("/NNEDenoiser/NNEDNN_Oidn2-2-0_Color_Dyn.NNEDNN_Oidn2-2-0_Color_Dyn");
-		case 2: return TEXT("/NNEDenoiser/NNEDNN_Oidn2-2-0_ColorAlbedoNormal_Dyn.NNEDNN_Oidn2-2-0_ColorAlbedoNormal_Dyn");
-		case 3: return TEXT("/NNEDenoiser/NNEDNN_Oidn2-2-0_HighQuality_Dyn.NNEDNN_Oidn2-2-0_HighQuality_Dyn");
-		
-		// Float16
-		case 4: return TEXT("/NNEDenoiser/NNEDNN_Oidn2-2-0_Color_Dyn_Float16.NNEDNN_Oidn2-2-0_Color_Dyn_Float16");
-		case 5: return TEXT("/NNEDenoiser/NNEDNN_Oidn2-2-0_ColorAlbedoNormal_Dyn_Float16.NNEDNN_Oidn2-2-0_ColorAlbedoNormal_Dyn_Float16");
-		case 6: return TEXT("/NNEDenoiser/NNEDNN_Oidn2-2-0_HighQuality_Dyn_Float16.NNEDNN_Oidn2-2-0_HighQuality_Dyn_Float16");
-
-		// Alpha, Float32
-		case 7: return TEXT("/NNEDenoiser/NNEDNN_Oidn2-2-0_Color_Dyn_Alpha.NNEDNN_Oidn2-2-0_Color_Dyn_Alpha");
-		case 8: return TEXT("/NNEDenoiser/NNEDNN_Oidn2-2-0_ColorAlbedoNormal_Dyn_Alpha.NNEDNN_Oidn2-2-0_ColorAlbedoNormal_Dyn_Alpha");
-		case 9: return TEXT("/NNEDenoiser/NNEDNN_Oidn2-2-0_HighQuality_Dyn_Alpha.NNEDNN_Oidn2-2-0_HighQuality_Dyn_Alpha");
-		
-		// Alpha, Float16
-		case 10: return TEXT("/NNEDenoiser/NNEDNN_Oidn2-2-0_Color_Dyn_Alpha_Float16.NNEDNN_Oidn2-2-0_Color_Dyn_Alpha_Float16");
-		case 11: return TEXT("/NNEDenoiser/NNEDNN_Oidn2-2-0_ColorAlbedoNormal_Dyn_Alpha_Float16.NNEDNN_Oidn2-2-0_ColorAlbedoNormal_Dyn_Alpha_Float16");
-		case 12: return TEXT("/NNEDenoiser/NNEDNN_Oidn2-2-0_HighQuality_Dyn_Alpha_Float16.NNEDNN_Oidn2-2-0_HighQuality_Dyn_Alpha_Float16");
+		case 1: return TEXT("/NNEDenoiser/NNEDNN_Oidn2-3-0_ColorAlbedoNormal_Small.NNEDNN_Oidn2-3-0_ColorAlbedoNormal_Small");
+		case 2: return TEXT("/NNEDenoiser/NNEDNN_Oidn2-3-0_ColorAlbedoNormal.NNEDNN_Oidn2-3-0_ColorAlbedoNormal");
+		case 3: return TEXT("/NNEDenoiser/NNEDNN_Oidn2-3-0_HighQuality_Large.NNEDNN_Oidn2-3-0_HighQuality_Large");
+
+		// Alpha
+		case 4: return TEXT("/NNEDenoiser/NNEDNN_Oidn2-3-0_ColorAlbedoNormal_Small_Alpha.NNEDNN_Oidn2-3-0_ColorAlbedoNormal_Small_Alpha");
+		case 5: return TEXT("/NNEDenoiser/NNEDNN_Oidn2-3-0_ColorAlbedoNormal_Alpha.NNEDNN_Oidn2-3-0_ColorAlbedoNormal_Alpha");
+		case 6: return TEXT("/NNEDenoiser/NNEDNN_Oidn2-3-0_HighQuality_Large_Alpha.NNEDNN_Oidn2-3-0_HighQuality_Large_Alpha");
 	}
 	check(false);
 	return FString();
