[NNEDenoiser] More consistent naming of denoiser assets.

#tests Editor, Win64 verified all models load correctly

#rnx

#virtualized

diff --git a/Engine/Plugins/Experimental/NNEDenoiser/Source/NNEDenoiser/Private/NNEDenoiserSettings.cpp b/Engine/Plugins/Experimental/NNEDenoiser/Source/NNEDenoiser/Private/NNEDenoiserSettings.cpp
--- a/Engine/Plugins/Experimental/NNEDenoiser/Source/NNEDenoiser/Private/NNEDenoiserSettings.cpp
+++ b/Engine/Plugins/Experimental/NNEDenoiser/Source/NNEDenoiser/Private/NNEDenoiserSettings.cpp
@@ -11,7 +11,7 @@ UNNEDenoiserSettings::UNNEDenoiserSettings()
 	CategoryName = TEXT("Plugins");
 	SectionName = TEXT("NNE Denoiser");
 
-	DenoiserModelData = FSoftObjectPath(TEXT("/NNEDenoiser/NNEDNN_Oidn2-3-0_ColorAlbedoNormal.NNEDNN_Oidn2-3-0_ColorAlbedoNormal"));
+	DenoiserModelData = FSoftObjectPath(TEXT("/NNEDenoiser/NNED_Oidn2-3_Balanced.NNED_Oidn2-3_Balanced"));
 }
 
 void UNNEDenoiserSettings::PostInitProperties()
diff --git a/Engine/Plugins/Experimental/NNEDenoiser/Source/NNEDenoiser/Private/NNEDenoiserViewExtension.cpp b/Engine/Plugins/Experimental/NNEDenoiser/Source/NNEDenoiser/Private/NNEDenoiserViewExtension.cpp
--- a/Engine/Plugins/Experimental/NNEDenoiser/Source/NNEDenoiser/Private/NNEDenoiserViewExtension.cpp
+++ b/Engine/Plugins/Experimental/NNEDenoiser/Source/NNEDenoiser/Private/NNEDenoiserViewExtension.cpp
@@ -65,19 +65,19 @@ EDenoiserRuntimeType GetDenoiserRuntimeTypeFromCVar()
 
 FString GetDenoiserModelDataNameFromCVarAndSettings(const UNNEDenoiserSettings* Settings)
 {
-	const int32 Idx = CVarNNEDenoiserModelData.GetValueOnGameThread();
+	const int32 Idx = FMath::Clamp(CVarNNEDenoiserModelData.GetValueOnGameThread(), 0, 6);
 	switch(Idx)
 	{
 		case 0: return !Settings->DenoiserModelData.IsNull() ? Settings->DenoiserModelData.ToString() : FString();
 
-		case 1: return TEXT("/NNEDenoiser/NNEDNN_Oidn2-3-0_ColorAlbedoNormal_Small.NNEDNN_Oidn2-3-0_ColorAlbedoNormal_Small");
-		case 2: return TEXT("/NNEDenoiser/NNEDNN_Oidn2-3-0_ColorAlbedoNormal.NNEDNN_Oidn2-3-0_ColorAlbedoNormal");
-		case 3: return TEXT("/NNEDenoiser/NNEDNN_Oidn2-3-0_HighQuality_Large.NNEDNN_Oidn2-3-0_HighQuality_Large");
+		case 1: return TEXT("/NNEDenoiser/NNED_Oidn2-3_Fast.NNED_Oidn2-3_Fast");
+		case 2: return TEXT("/NNEDenoiser/NNED_Oidn2-3_Balanced.NNED_Oidn2-3_Balanced");
+		case 3: return TEXT("/NNEDenoiser/NNED_Oidn2-3_HighQuality.NNED_Oidn2-3_HighQuality");
 
 		// Alpha
-		case 4: return TEXT("/NNEDenoiser/NNEDNN_Oidn2-3-0_ColorAlbedoNormal_Small_Alpha.NNEDNN_Oidn2-3-0_ColorAlbedoNormal_Small_Alpha");
-		case 5: return TEXT("/NNEDenoiser/NNEDNN_Oidn2-3-0_ColorAlbedoNormal_Alpha.NNEDNN_Oidn2-3-0_ColorAlbedoNormal_Alpha");
-		case 6: return TEXT("/NNEDenoiser/NNEDNN_Oidn2-3-0_HighQuality_Large_Alpha.NNEDNN_Oidn2-3-0_HighQuality_Large_Alpha");
+		case 4: return TEXT("/NNEDenoiser/NNED_Oidn2-3_Fast_Alpha.NNED_Oidn2-3_Fast_Alpha");
+		case 5: return TEXT("/NNEDenoiser/NNED_Oidn2-3_Balanced_Alpha.NNED_Oidn2-3_Balanced_Alpha");
+		case 6: return TEXT("/NNEDenoiser/NNED_Oidn2-3_HighQuality_Alpha.NNED_Oidn2-3_HighQuality_Alpha");
 	}
 	check(false);
 	return FString();
