[NNE] Deprecate FTensorBindingGPU in favor of FTensorBindingCPU

Problem: FTensorBindingGPU was replaced by FTensorBindingCPU but not deprecated yet.
Solution: Deprecate it and update internal usage of in the engine and plugins.

Tests: editor win64 + packaged win64 : all nne tests

#rnx

diff --git a/Engine/Plugins/Experimental/NNEDenoiser/Source/NNEDenoiser/Private/NNEDenoiserModelInstanceGPU.cpp b/Engine/Plugins/Experimental/NNEDenoiser/Source/NNEDenoiser/Private/NNEDenoiserModelInstanceGPU.cpp
--- a/Engine/Plugins/Experimental/NNEDenoiser/Source/NNEDenoiser/Private/NNEDenoiserModelInstanceGPU.cpp
+++ b/Engine/Plugins/Experimental/NNEDenoiser/Source/NNEDenoiser/Private/NNEDenoiserModelInstanceGPU.cpp
@@ -126,21 +126,21 @@ namespace UE::NNEDenoiser::Private
 #endif
 
 			ScratchInputBuffers.SetNum(DenoiserParameters->InputBuffers.Num());
-			TArray<NNE::FTensorBindingGPU> InputBindings;
+			TArray<NNE::FTensorBindingCPU> InputBindings;
 			for (int32 Idx = 0; Idx < DenoiserParameters->InputBuffers.Num(); Idx++)
 			{
 				ScratchInputBuffers[Idx].SetNumUninitialized(DenoiserParameters->InputBuffers[Idx].GetBuffer()->GetRHI()->GetSize());
 
-				InputBindings.Emplace(NNE::FTensorBindingGPU{(void *)ScratchInputBuffers[Idx].GetData(), (uint64)ScratchInputBuffers[Idx].Num()});
+				InputBindings.Emplace(NNE::FTensorBindingCPU{(void *)ScratchInputBuffers[Idx].GetData(), (uint64)ScratchInputBuffers[Idx].Num()});
 			}
 
 			ScratchOutputBuffers.SetNum(DenoiserParameters->OutputBuffers.Num());
-			TArray<NNE::FTensorBindingGPU> OutputBindings;
+			TArray<NNE::FTensorBindingCPU> OutputBindings;
 			for (int32 Idx = 0; Idx < DenoiserParameters->OutputBuffers.Num(); Idx++)
 			{
 				ScratchOutputBuffers[Idx].SetNumUninitialized(DenoiserParameters->OutputBuffers[Idx].GetBuffer()->GetRHI()->GetSize());
 
-				OutputBindings.Emplace(NNE::FTensorBindingGPU{(void *)ScratchOutputBuffers[Idx].GetData(), (uint64)ScratchOutputBuffers[Idx].Num()});
+				OutputBindings.Emplace(NNE::FTensorBindingCPU{(void *)ScratchOutputBuffers[Idx].GetData(), (uint64)ScratchOutputBuffers[Idx].Num()});
 			}
 
 			for (int32 Idx = 0; Idx < DenoiserParameters->InputBuffers.Num(); Idx++)
diff --git a/Engine/Plugins/NNE/NNERuntimeORT/Source/NNERuntimeORT/Private/NNERuntimeORTModel.h b/Engine/Plugins/NNE/NNERuntimeORT/Source/NNERuntimeORT/Private/NNERuntimeORTModel.h
--- a/Engine/Plugins/NNE/NNERuntimeORT/Source/NNERuntimeORT/Private/NNERuntimeORTModel.h
+++ b/Engine/Plugins/NNE/NNERuntimeORT/Source/NNERuntimeORT/Private/NNERuntimeORTModel.h
@@ -98,7 +98,7 @@ private:
 	TSharedPtr<UE::NNE::FSharedModelData> ModelData;
 };
 
-class FModelInstanceORTDmlGPU : public FModelInstanceORTBase<NNE::IModelInstanceGPU, NNE::FTensorBindingGPU>
+class FModelInstanceORTDmlGPU : public FModelInstanceORTBase<NNE::IModelInstanceGPU, NNE::FTensorBindingCPU>
 {
 public:
 	FModelInstanceORTDmlGPU(const FRuntimeConf& InRuntimeConf, TSharedRef<FEnvironment> InEnvironment) : FModelInstanceORTBase(InRuntimeConf, InEnvironment) {}
diff --git a/Engine/Source/Runtime/NNE/Public/NNERuntimeGPU.h b/Engine/Source/Runtime/NNE/Public/NNERuntimeGPU.h
--- a/Engine/Source/Runtime/NNE/Public/NNERuntimeGPU.h
+++ b/Engine/Source/Runtime/NNE/Public/NNERuntimeGPU.h
@@ -12,6 +12,7 @@ class UNNEModelData;
 
 namespace UE::NNE
 {
+UE_DEPRECATED(5.5, "Use UE::NNE::FTensorBindingCPU instead.")
 typedef FTensorBindingCPU FTensorBindingGPU;
 
 /**
