[NNEDenoiser] Minor fixes and updates:
* Renamed NNEDenoiserModelData to NNEDenoiserAsset
* "Alignment" might be misleading, name it instead "Size Alignment"
* Adjust paragraph name from "create" to "select", because a user can SELECT a preset, not create one

#tests Editor, Win64 run every model

#rnx

#virtualized

diff --git a/Engine/Plugins/Experimental/NNEDenoiser/Source/NNEDenoiser/Private/NNEDenoiserSettings.cpp b/Engine/Plugins/Experimental/NNEDenoiser/Source/NNEDenoiser/Private/NNEDenoiserSettings.cpp
--- a/Engine/Plugins/Experimental/NNEDenoiser/Source/NNEDenoiser/Private/NNEDenoiserSettings.cpp
+++ b/Engine/Plugins/Experimental/NNEDenoiser/Source/NNEDenoiser/Private/NNEDenoiserSettings.cpp
@@ -11,7 +11,7 @@ UNNEDenoiserSettings::UNNEDenoiserSettings()
 	CategoryName = TEXT("Plugins");
 	SectionName = TEXT("NNE Denoiser");
 
-	DenoiserModelData = FSoftObjectPath(TEXT("/NNEDenoiser/NNED_Oidn2-3_Balanced.NNED_Oidn2-3_Balanced"));
+	DenoiserAsset = FSoftObjectPath(TEXT("/NNEDenoiser/NNED_Oidn2-3_Balanced.NNED_Oidn2-3_Balanced"));
 }
 
 void UNNEDenoiserSettings::PostInitProperties()
diff --git a/Engine/Plugins/Experimental/NNEDenoiser/Source/NNEDenoiser/Private/NNEDenoiserViewExtension.cpp b/Engine/Plugins/Experimental/NNEDenoiser/Source/NNEDenoiser/Private/NNEDenoiserViewExtension.cpp
--- a/Engine/Plugins/Experimental/NNEDenoiser/Source/NNEDenoiser/Private/NNEDenoiserViewExtension.cpp
+++ b/Engine/Plugins/Experimental/NNEDenoiser/Source/NNEDenoiser/Private/NNEDenoiserViewExtension.cpp
@@ -6,7 +6,7 @@
 #include "NNEDenoiserGenericDenoiser.h"
 #include "NNEDenoiserIOProcessBase.h"
 #include "NNEDenoiserLog.h"
-#include "NNEDenoiserModelData.h"
+#include "NNEDenoiserAsset.h"
 #include "NNEDenoiserModelInstanceCPU.h"
 #include "NNEDenoiserModelInstanceGPU.h"
 #include "NNEDenoiserModelInstanceRDG.h"
@@ -27,11 +27,11 @@ static TAutoConsoleVariable<bool> CVarNNEDenoiser(
 	TEXT("Enable the NNE Denoiser.")
 );
 
-static TAutoConsoleVariable<int32> CVarNNEDenoiserModelData(
-	TEXT("NNEDenoiser.ModelData"),
+static TAutoConsoleVariable<int32> CVarNNEDenoiserAsset(
+	TEXT("NNEDenoiser.Asset"),
 	0,
-	TEXT("Defines the model data of the denoiser model.\n")
-	TEXT("  0: Use model data from Project Settings\n")
+	TEXT("Defines asset to used to create the denoiser.\n")
+	TEXT("  0: Use denoiser asset defined by Project Settings\n")
 	TEXT("  1: OIDN Fast\n")
 	TEXT("  2: OIDN Balanced\n")
 	TEXT("  3: OIDN High Quality\n")
@@ -63,12 +63,12 @@ EDenoiserRuntimeType GetDenoiserRuntimeTypeFromCVar()
 	return static_cast<EDenoiserRuntimeType>(FMath::Clamp(Value, Min, Max));
 }
 
-FString GetDenoiserModelDataNameFromCVarAndSettings(const UNNEDenoiserSettings* Settings)
+FString GetDenoiserAssetNameFromCVarAndSettings(const UNNEDenoiserSettings* Settings)
 {
-	const int32 Idx = FMath::Clamp(CVarNNEDenoiserModelData.GetValueOnGameThread(), 0, 6);
+	const int32 Idx = FMath::Clamp(CVarNNEDenoiserAsset.GetValueOnGameThread(), 0, 6);
 	switch(Idx)
 	{
-		case 0: return !Settings->DenoiserModelData.IsNull() ? Settings->DenoiserModelData.ToString() : FString();
+		case 0: return !Settings->DenoiserAsset.IsNull() ? Settings->DenoiserAsset.ToString() : FString();
 
 		case 1: return TEXT("/NNEDenoiser/NNED_Oidn2-3_Fast.NNED_Oidn2-3_Fast");
 		case 2: return TEXT("/NNEDenoiser/NNED_Oidn2-3_Balanced.NNED_Oidn2-3_Balanced");
@@ -207,15 +207,15 @@ static FResourceMappingList MakeTensorLayout(UDataTable* DataTable)
 	return Result;
 }
 
-FParameters GetParametersValidated(const UNNEDenoiserModelData& DenoiserModelData)
+FParameters GetParametersValidated(const UNNEDenoiserAsset& DenoiserAsset)
 {
 	FParameters Parameters{
 		.TilingConfig =
 		{
-			.Alignment = DenoiserModelData.TilingConfig.Alignment,
-			.Overlap = DenoiserModelData.TilingConfig.Overlap,
-			.MaxSize = DenoiserModelData.TilingConfig.MaxSize,
-			.MinSize = DenoiserModelData.TilingConfig.MinSize
+			.Alignment = DenoiserAsset.TilingConfig.Alignment,
+			.Overlap = DenoiserAsset.TilingConfig.Overlap,
+			.MaxSize = DenoiserAsset.TilingConfig.MaxSize,
+			.MinSize = DenoiserAsset.TilingConfig.MinSize
 		}
 	};
 
@@ -266,28 +266,28 @@ TUniquePtr<FGenericDenoiser> CreateNNEDenoiserFromAsset(const FString& AssetName
 		return {};
 	}
 
-	UNNEDenoiserModelData *DenoiserModelData = LoadObject<UNNEDenoiserModelData>(nullptr, *AssetName);
-	if (!DenoiserModelData)
+	UNNEDenoiserAsset *DenoiserAsset = LoadObject<UNNEDenoiserAsset>(nullptr, *AssetName);
+	if (!DenoiserAsset)
 	{
 		UE_LOG(LogNNEDenoiser, Error, TEXT("Could not load denoiser model data asset!"));
 		return {};
 	}
 
-	UNNEModelData* ModelData = DenoiserModelData->ModelData.LoadSynchronous();
+	UNNEModelData* ModelData = DenoiserAsset->ModelData.LoadSynchronous();
 	if (!ModelData)
 	{
 		UE_LOG(LogNNEDenoiser, Error, TEXT("Asset does not contain model data!"));
 		return {};
 	}
-	UDataTable* InputMappingTable = DenoiserModelData->InputMapping.LoadSynchronous();
+	UDataTable* InputMappingTable = DenoiserAsset->InputMapping.LoadSynchronous();
 	if (InputMappingTable)
 	{
-		UE_LOG(LogNNEDenoiser, Log, TEXT("Loaded input mapping from %s"), *DenoiserModelData->InputMapping.GetAssetName());
+		UE_LOG(LogNNEDenoiser, Log, TEXT("Loaded input mapping from %s"), *DenoiserAsset->InputMapping.GetAssetName());
 	}
-	UDataTable* OutputMappingTable = DenoiserModelData->OutputMapping.LoadSynchronous();
+	UDataTable* OutputMappingTable = DenoiserAsset->OutputMapping.LoadSynchronous();
 	if (OutputMappingTable)
 	{
-		UE_LOG(LogNNEDenoiser, Log, TEXT("Loaded output mapping from %s"), *DenoiserModelData->OutputMapping.GetAssetName());
+		UE_LOG(LogNNEDenoiser, Log, TEXT("Loaded output mapping from %s"), *DenoiserAsset->OutputMapping.GetAssetName());
 	}
 
 	FResourceMappingList InputLayout;
@@ -318,7 +318,7 @@ TUniquePtr<FGenericDenoiser> CreateNNEDenoiserFromAsset(const FString& AssetName
 	TUniquePtr<IInputProcess> InputProcess = MakeUnique<FInputProcessBase>(MoveTemp(InputLayout), TransferFunction);
 	TUniquePtr<IOutputProcess> OutputProcess = MakeUnique<FOutputProcessBase>(MoveTemp(OutputLayout), TransferFunction);
 
-	FParameters Parameters = GetParametersValidated(*DenoiserModelData);
+	FParameters Parameters = GetParametersValidated(*DenoiserAsset);
 
 	return CreateNNEDenoiser(
 		*ModelData,
@@ -388,7 +388,7 @@ void FViewExtension::ApplySettings(const UNNEDenoiserSettings* Settings)
 	bNeedsUpdate |= CVarNNEDenoiser.GetValueOnGameThread() != bDenoiserEnabled;
 	bNeedsUpdate |= GetDenoiserRuntimeTypeFromCVar() != RuntimeType;
 	bNeedsUpdate |= CVarNNEDenoiserRuntimeName.GetValueOnGameThread() != RuntimeName;
-	bNeedsUpdate |= GetDenoiserModelDataNameFromCVarAndSettings(Settings) != ModelDataName;
+	bNeedsUpdate |= GetDenoiserAssetNameFromCVarAndSettings(Settings) != ModelDataName;
 
 	if (!bNeedsUpdate)
 	{
@@ -406,7 +406,7 @@ void FViewExtension::ApplySettings(const UNNEDenoiserSettings* Settings)
 
 	RuntimeType = GetDenoiserRuntimeTypeFromCVar();
 	RuntimeName = CVarNNEDenoiserRuntimeName.GetValueOnGameThread();
-	ModelDataName = GetDenoiserModelDataNameFromCVarAndSettings(Settings);
+	ModelDataName = GetDenoiserAssetNameFromCVarAndSettings(Settings);
 
 	UE_LOG(LogNNEDenoiser, Log, TEXT("Create denoiser from asset %s..."), *ModelDataName);
 
diff --git a/Engine/Plugins/Experimental/NNEDenoiser/Source/NNEDenoiser/Public/NNEDenoiserAsset.h b/Engine/Plugins/Experimental/NNEDenoiser/Source/NNEDenoiser/Public/NNEDenoiserAsset.h
new file mode 100644
--- /dev/null
+++ b/Engine/Plugins/Experimental/NNEDenoiser/Source/NNEDenoiser/Public/NNEDenoiserAsset.h
@@ -0,0 +1,56 @@
+// Copyright Epic Games, Inc. All Rights Reserved.
+
+#pragma once
+
+#include "Engine/DataAsset.h"
+#include "Engine/DataTable.h"
+#include "NNEModelData.h"
+
+#include "NNEDenoiserAsset.generated.h"
+
+/** Tiling configuration for fixed and dynamic size models */
+USTRUCT(BlueprintType)
+struct FTilingConfig
+{
+	GENERATED_BODY()
+
+	/** Tile size alignment (applies only to dynamic size models) */
+	UPROPERTY(EditAnywhere, BlueprintReadWrite, Category=NNEDenoiser, meta = (DisplayName = "Size Alignment"))
+	int32 Alignment = 1;
+
+	/** Tile overlap */
+	UPROPERTY(EditAnywhere, BlueprintReadWrite, Category=NNEDenoiser)
+	int32 Overlap = 0;
+
+	/** Maximum tile size (applies only to dynamic size models) */
+	UPROPERTY(EditAnywhere, BlueprintReadWrite, Category=NNEDenoiser)
+	int32 MaxSize = 0;
+
+	/** Minimum tile size (applies only to dynamic size models) */
+	UPROPERTY(EditAnywhere, BlueprintReadWrite, Category=NNEDenoiser)
+	int32 MinSize = 1;
+};
+
+/** Denoiser model data asset */
+UCLASS(BlueprintType)
+class NNEDENOISER_API UNNEDenoiserAsset : public UDataAsset
+{
+	GENERATED_BODY()
+
+public:
+	/** NNE model data */
+	UPROPERTY(EditAnywhere, BlueprintReadWrite, Category=NNEDenoiser)
+	TSoftObjectPtr<UNNEModelData> ModelData;
+
+	/** Input mapping table */
+	UPROPERTY(EditAnywhere, BlueprintReadWrite, Category=NNEDenoiser, meta = (RequiredAssetDataTags = "RowStructure=/Script/NNEDenoiser.NNEDenoiserModelIOMappingData"))
+	TSoftObjectPtr<UDataTable> InputMapping;
+
+	/** Output mapping table */
+	UPROPERTY(EditAnywhere, BlueprintReadWrite, Category=NNEDenoiser, meta = (RequiredAssetDataTags = "RowStructure=/Script/NNEDenoiser.NNEDenoiserModelIOMappingData"))
+	TSoftObjectPtr<UDataTable> OutputMapping;
+
+	/** Tiling configuration */
+	UPROPERTY(EditAnywhere, BlueprintReadWrite, Category=NNEDenoiser)
+	FTilingConfig TilingConfig{};
+};
\ No newline at end of file
diff --git a/Engine/Plugins/Experimental/NNEDenoiser/Source/NNEDenoiser/Public/NNEDenoiserModelData.h b/Engine/Plugins/Experimental/NNEDenoiser/Source/NNEDenoiser/Public/NNEDenoiserModelData.h
deleted file mode 100644
diff --git a/Engine/Plugins/Experimental/NNEDenoiser/Source/NNEDenoiser/Public/NNEDenoiserSettings.h b/Engine/Plugins/Experimental/NNEDenoiser/Source/NNEDenoiser/Public/NNEDenoiserSettings.h
--- a/Engine/Plugins/Experimental/NNEDenoiser/Source/NNEDenoiser/Public/NNEDenoiserSettings.h
+++ b/Engine/Plugins/Experimental/NNEDenoiser/Source/NNEDenoiser/Public/NNEDenoiserSettings.h
@@ -3,7 +3,7 @@
 #pragma once
 
 #include "Engine/DeveloperSettingsBackedByCVars.h"
-#include "NNEDenoiserModelData.h"
+#include "NNEDenoiserAsset.h"
 
 #include "NNEDenoiserSettings.generated.h"
 
@@ -16,7 +16,7 @@ enum EDenoiserRuntimeType : uint8
 	RDG
 };
 
-/** Settings used to create a NNE Denoiser */
+/** Settings to select a NNE Denoiser and its runtime */
 UCLASS(Config = Engine, meta = (DisplayName = "NNE Denoiser"))
 class NNEDENOISER_API UNNEDenoiserSettings : public UDeveloperSettingsBackedByCVars
 {
@@ -27,9 +27,9 @@ public:
 
 	virtual void PostInitProperties() override;
 
-	/** Denoiser model data used to create a NNE Denoiser */
-	UPROPERTY(Config, EditAnywhere, Category = "NNE Denoiser", meta = (DisplayName = "Denoiser model data", ToolTip = "Select the denoiser model data"))
-	TSoftObjectPtr<UNNEDenoiserModelData> DenoiserModelData;
+	/** Denoiser asset data used to create a NNE Denoiser */
+	UPROPERTY(Config, EditAnywhere, Category = "NNE Denoiser", meta = (DisplayName = "Denoiser Asset", ToolTip = "Select the denoiser asset"))
+	TSoftObjectPtr<UNNEDenoiserAsset> DenoiserAsset;
 
 private:
 	/** Runtime type used to run the NNE Denoiser model. Backed by the console variable 'NNEDenoiser.Runtime.Type'. */
