[NNEDenoiser] Register also temporal denoiser:
* Changed name used to register to "NNEDenoiser" (for both)
* Also register temporal densoiser with name "NNEDenoiser"
* Added temporal denoiser preset dialog to project settings
* Added temporal denoiser preset data asset type
* Splited mapping asset tables to use separat row structs for intput and output
* Added separate input/output mapping row struct types for temporal denoiser
* Added meta tag in denoiser preset data asset to filter for specific row structs

#tests Editor, Win64 load and use all denosier assets

#rnx

#virtualized

diff --git a/Engine/Plugins/Experimental/NNEDenoiser/Source/NNEDenoiser/Private/NNEDenoiserHistory.h b/Engine/Plugins/Experimental/NNEDenoiser/Source/NNEDenoiser/Private/NNEDenoiserHistory.h
--- a/Engine/Plugins/Experimental/NNEDenoiser/Source/NNEDenoiser/Private/NNEDenoiserHistory.h
+++ b/Engine/Plugins/Experimental/NNEDenoiser/Source/NNEDenoiser/Private/NNEDenoiserHistory.h
@@ -4,7 +4,7 @@
 
 #include "CoreMinimal.h"
 #include "NNEDenoiserLog.h"
-#include "NNEDenoiserModelIOMappingData.h"
+#include "NNEDenoiserResourceName.h"
 #include "PathTracingDenoiser.h"
 #include "RendererInterface.h"
 
diff --git a/Engine/Plugins/Experimental/NNEDenoiser/Source/NNEDenoiser/Private/NNEDenoiserIOMappingData.cpp b/Engine/Plugins/Experimental/NNEDenoiser/Source/NNEDenoiser/Private/NNEDenoiserIOMappingData.cpp
new file mode 100644
--- /dev/null
+++ b/Engine/Plugins/Experimental/NNEDenoiser/Source/NNEDenoiser/Private/NNEDenoiserIOMappingData.cpp
@@ -0,0 +1,53 @@
+#include "NNEDenoiserIOMappingData.h"
+
+namespace UE::NNEDenoiser
+{
+	
+EResourceName ToResourceName(EInputResourceName Name)
+{
+	switch(Name)
+	{
+		case EInputResourceName::I_Color: return EResourceName::Color;
+		case EInputResourceName::I_Albedo: return EResourceName::Albedo;
+		case EInputResourceName::I_Normal: return EResourceName::Normal;
+		case EInputResourceName::I_Output: return EResourceName::Output;
+		default: checkNoEntry();
+	}
+	return EResourceName::Color;
+}
+
+EResourceName ToResourceName(EOutputResourceName Name)
+{
+	switch(Name)
+	{
+		case EOutputResourceName::O_Output: return EResourceName::Output;
+		default: checkNoEntry();
+	}
+	return EResourceName::Color;
+}
+
+EResourceName ToResourceName(ETemporalInputResourceName Name)
+{
+	switch(Name)
+	{
+		case ETemporalInputResourceName::TI_Color: return EResourceName::Color;
+		case ETemporalInputResourceName::TI_Albedo: return EResourceName::Albedo;
+		case ETemporalInputResourceName::TI_Normal: return EResourceName::Normal;
+		case ETemporalInputResourceName::TI_Flow: return EResourceName::Flow;
+		case ETemporalInputResourceName::TI_Output: return EResourceName::Output;
+		default: checkNoEntry();
+	}
+	return EResourceName::Color;
+}
+
+EResourceName ToResourceName(ETemporalOutputResourceName Name)
+{
+	switch(Name)
+	{
+		case ETemporalOutputResourceName::TO_Output: return EResourceName::Output;
+		default: checkNoEntry();
+	}
+	return EResourceName::Color;
+}
+
+} // namespace UE::NNEDenoiser
\ No newline at end of file
diff --git a/Engine/Plugins/Experimental/NNEDenoiser/Source/NNEDenoiser/Private/NNEDenoiserIOProcess.h b/Engine/Plugins/Experimental/NNEDenoiser/Source/NNEDenoiser/Private/NNEDenoiserIOProcess.h
--- a/Engine/Plugins/Experimental/NNEDenoiser/Source/NNEDenoiser/Private/NNEDenoiserIOProcess.h
+++ b/Engine/Plugins/Experimental/NNEDenoiser/Source/NNEDenoiser/Private/NNEDenoiserIOProcess.h
@@ -2,7 +2,7 @@
 
 #pragma once
 
-#include "NNEDenoiserModelIOMappingData.h"
+#include "NNEDenoiserResourceName.h"
 #include "NNETypes.h"
 #include "RenderGraphFwd.h"
 
diff --git a/Engine/Plugins/Experimental/NNEDenoiser/Source/NNEDenoiser/Private/NNEDenoiserIOProcessBase.cpp b/Engine/Plugins/Experimental/NNEDenoiser/Source/NNEDenoiser/Private/NNEDenoiserIOProcessBase.cpp
--- a/Engine/Plugins/Experimental/NNEDenoiser/Source/NNEDenoiser/Private/NNEDenoiserIOProcessBase.cpp
+++ b/Engine/Plugins/Experimental/NNEDenoiser/Source/NNEDenoiser/Private/NNEDenoiserIOProcessBase.cpp
@@ -84,7 +84,7 @@ void AddReadInputPass(
 	FRDGBuilder& GraphBuilder,
 	FRDGTextureRef InputTexture,
 	FRDGBufferUAVRef OutputBufferUAV,
-	const TArray<FIntPoint>& ChannelMapping)
+	const TArray<FChannelMapping>& ChannelMapping)
 {
 	using namespace UE::NNEDenoiserShaders::Internal;
 
@@ -109,7 +109,7 @@ void AddReadInputPass(
 	ReadInputParameters->OutputBuffer = OutputBufferUAV;
 	for (int32 Idx = 0; Idx < ChannelMapping.Num(); Idx++)
 	{
-		ReadInputParameters->OutputChannel_InputChannel_Unused_Unused[Idx] = { ChannelMapping[Idx].X, ChannelMapping[Idx].Y, 0, 0 };
+		ReadInputParameters->OutputChannel_InputChannel_Unused_Unused[Idx] = { ChannelMapping[Idx].TensorChannel, ChannelMapping[Idx].ResourceChannel, 0, 0 };
 	}
 
 	const EDataType InputDataType = GetDenoiserShaderDataType(InputTexture->Desc.Format);
@@ -146,7 +146,7 @@ void AddWriteOutputPass(
 	EPixelFormat BufferFormat,
 	FRDGTextureRef OutputTexture,
 	UE::NNEDenoiserShaders::Internal::EDataType DataType,
-	const TArray<FIntPoint>& ChannelMapping)
+	const TArray<FChannelMapping>& ChannelMapping)
 {
 	using namespace UE::NNEDenoiserShaders::Internal;
 
@@ -159,7 +159,7 @@ void AddWriteOutputPass(
 	WriteOutputParameters->OutputTexture = GraphBuilder.CreateUAV(FRDGTextureUAVDesc(OutputTexture));
 	for (int32 Idx = 0; Idx < ChannelMapping.Num(); Idx++)
 	{
-		WriteOutputParameters->OutputChannel_InputChannel_Unused_Unused[Idx] = { ChannelMapping[Idx].Y, ChannelMapping[Idx].X, 0, 0 };
+		WriteOutputParameters->OutputChannel_InputChannel_Unused_Unused[Idx] = { ChannelMapping[Idx].ResourceChannel, ChannelMapping[Idx].TensorChannel, 0, 0 };
 	}
 
 	const EDataType InputDataType = GetDenoiserShaderDataType(BufferFormat);
@@ -197,7 +197,7 @@ void AddReadInputPassForKind(
 	const FResourceMapping& ResourceMapping,
 	FRDGBufferUAVRef BufferUAV)
 {
-	for (const TPair<int32, TArray<FIntPoint>>& ChannelMappingKeyValue : ResourceMapping.GetChannelMappingPerFrame(TensorName))
+	for (const TPair<int32, TArray<FChannelMapping>>& ChannelMappingKeyValue : ResourceMapping.GetChannelMappingPerFrame(TensorName))
 	{
 		const int32 FrameIdx = -ChannelMappingKeyValue.Key;
 		FRDGTextureRef InputTexture = ResourceAccess.GetTexture(TensorName, FrameIdx);
@@ -498,9 +498,9 @@ void FOutputProcessBase::ReadOutputBuffer(
 	const NNEDenoiserShaders::Internal::EDataType DataType = GetDenoiserShaderDataType(TensorDataType);
 	const FIntPoint BufferSize(TensorShape.GetData()[3], TensorShape.GetData()[2]);
 		
-	TMap<int32, TArray<FIntPoint>> ChannelMappingPerFrame = ResourceMapping.GetChannelMappingPerFrame(EResourceName::Output);
+	TMap<int32, TArray<FChannelMapping>> ChannelMappingPerFrame = ResourceMapping.GetChannelMappingPerFrame(EResourceName::Output);
 	checkf(ChannelMappingPerFrame.Num() == 1, TEXT("Invalid output mapping"));
-	const TArray<FIntPoint>& ChannelMapping = ChannelMappingPerFrame.CreateConstIterator().Value();
+	const TArray<FChannelMapping>& ChannelMapping = ChannelMappingPerFrame.CreateConstIterator().Value();
 
 	const EPixelFormat BufferFormat = GetBufferFormat(TensorDataType);
 
diff --git a/Engine/Plugins/Experimental/NNEDenoiser/Source/NNEDenoiser/Private/NNEDenoiserResourceManager.h b/Engine/Plugins/Experimental/NNEDenoiser/Source/NNEDenoiser/Private/NNEDenoiserResourceManager.h
--- a/Engine/Plugins/Experimental/NNEDenoiser/Source/NNEDenoiser/Private/NNEDenoiserResourceManager.h
+++ b/Engine/Plugins/Experimental/NNEDenoiser/Source/NNEDenoiser/Private/NNEDenoiserResourceManager.h
@@ -3,7 +3,7 @@
 #pragma once
 
 #include "Containers/Map.h"
-#include "NNEDenoiserModelIOMappingData.h"
+#include "NNEDenoiserResourceName.h"
 #include "NNEDenoiserTiling.h"
 #include "RendererInterface.h"
 #include "RenderGraphFwd.h"
diff --git a/Engine/Plugins/Experimental/NNEDenoiser/Source/NNEDenoiser/Private/NNEDenoiserResourceMapping.cpp b/Engine/Plugins/Experimental/NNEDenoiser/Source/NNEDenoiser/Private/NNEDenoiserResourceMapping.cpp
new file mode 100644
--- /dev/null
+++ b/Engine/Plugins/Experimental/NNEDenoiser/Source/NNEDenoiser/Private/NNEDenoiserResourceMapping.cpp
@@ -0,0 +1,244 @@
+// Copyright Epic Games, Inc. All Rights Reserved.
+
+#include "NNEDenoiserResourceMapping.h"
+
+namespace UE::NNEDenoiser::Private
+{
+
+namespace Detail
+{
+
+FResourceMappingList MakeTensorLayout(const TMap<int32, TMap<int32, FResourceInfo>>& ResourceMap)
+{
+	FResourceMappingList Result{};
+
+	for (int32 I = 0; I < ResourceMap.Num(); I++)
+	{
+		checkf(ResourceMap.Contains(I), TEXT("Missing intput/output %d, must be continuous!"), I);
+
+		const TMap<int32, FResourceInfo>& InnerMap = ResourceMap[I];
+		FResourceMapping& Mapping = Result.Add_GetRef({});
+
+		for (int32 J = 0; J < InnerMap.Num(); J++)
+		{
+			checkf(InnerMap.Contains(J), TEXT("Missing tensor info for channel %d, must be continuous!"), J);
+
+			Mapping.Add(InnerMap[J]);
+		}
+	}
+
+	return Result;
+}
+
+
+void Add(TMap<int32, TMap<int32, FResourceInfo>>& Map, int32 TensorIndex, int32 TensorChannel, EResourceName ResourceName, int32 ResourceChannel, int32 FrameIndex)
+{
+	if (TensorChannel < 0)
+	{
+		for (int32 I = 0; I < -TensorChannel; I++)
+		{
+			Map.FindOrAdd(TensorIndex).FindOrAdd(I) = FResourceInfo{ResourceName, I, 0};
+		}
+	}
+	else
+	{
+		Map.FindOrAdd(TensorIndex).FindOrAdd(TensorChannel) = FResourceInfo{ResourceName, ResourceChannel, 0};
+	}
+}
+
+} // namespace Detail
+
+void FResourceMapping::Add(FResourceInfo Info)
+{
+	ChannelMapping.Add(MoveTemp(Info));
+}
+
+FResourceInfo& FResourceMapping::Add_GetRef(FResourceInfo Info)
+{
+	return ChannelMapping.Add_GetRef(MoveTemp(Info));
+}
+
+const FResourceInfo& FResourceMapping::GetChecked(int32 Channel) const
+{
+	check(Channel >= 0 && Channel < ChannelMapping.Num());
+
+	return ChannelMapping[Channel];
+}
+
+int32 FResourceMapping::Num() const
+{
+	return ChannelMapping.Num();
+}
+
+bool FResourceMapping::HasResource(EResourceName Name) const
+{
+	for (const auto& Info: ChannelMapping)
+	{
+		if (Info.Name == Name)
+		{
+			return true;
+		}
+	}
+	return false;
+}
+
+int32 FResourceMapping::NumFrames(EResourceName Name) const
+{
+	int32 MinFrameIndex = 1;
+	for (const auto& Info: ChannelMapping)
+	{
+		if (Info.Name == Name)
+		{
+			MinFrameIndex = FMath::Min(MinFrameIndex, Info.Frame);
+		}
+	}
+
+	return 1 - MinFrameIndex;
+}
+
+TMap<int32, TArray<FChannelMapping>> FResourceMapping::GetChannelMappingPerFrame(EResourceName Name) const
+{
+	TMap<int32, TArray<FChannelMapping>> Result;
+	for (int32 I = 0; I < ChannelMapping.Num(); I++)
+	{
+		const FResourceInfo& Info = ChannelMapping[I];
+
+		if (Info.Name == Name)
+		{
+			Result.FindOrAdd(Info.Frame).Emplace(I, Info.Channel);
+		}
+	}
+
+	return Result;
+}
+
+void FResourceMappingList::Add(FResourceMapping ResourceMapping)
+{
+	InputMapping.Add(MoveTemp(ResourceMapping));
+}
+
+FResourceMapping& FResourceMappingList::Add_GetRef(FResourceMapping ResourceMapping)
+{
+	return InputMapping.Add_GetRef(MoveTemp(ResourceMapping));
+}
+
+const FResourceMapping* FResourceMappingList::Get(int32 InputIndex) const
+{
+	if (InputIndex >= 0 && InputIndex < InputMapping.Num())
+	{
+		return &InputMapping[InputIndex];
+	}
+
+	return {};
+}
+
+const FResourceMapping& FResourceMappingList::GetChecked(int32 InputIndex) const
+{
+	check(InputIndex >= 0 && InputIndex < InputMapping.Num());
+
+	return InputMapping[InputIndex];
+}
+
+int32 FResourceMappingList::Num() const
+{
+	return InputMapping.Num();
+}
+
+int32 FResourceMappingList::NumChannels(int32 InputIndex) const
+{
+	check(InputIndex >= 0 && InputIndex < InputMapping.Num());
+	
+	return InputMapping[InputIndex].Num();
+}
+
+int32 FResourceMappingList::NumFrames(EResourceName Name) const
+{
+	int32 MaxNumFrames = 0;
+	for (const auto& Input : InputMapping)
+	{
+		const int32 NumFrames = Input.NumFrames(Name);
+		if (MaxNumFrames < NumFrames)
+		{
+			MaxNumFrames = NumFrames;
+		}
+	}
+
+	return MaxNumFrames;
+}
+
+bool FResourceMappingList::HasResource(EResourceName Name) const
+{
+	for (const auto& Input : InputMapping)
+	{
+		if (Input.HasResource(Name))
+		{
+			return true;
+		}
+	}
+
+	return false;
+}
+
+bool FResourceMappingList::HasResource(int32 InputIndex, EResourceName Name) const
+{
+	check(InputIndex >= 0 && InputIndex < InputMapping.Num());
+
+	return InputMapping[InputIndex].HasResource(Name);
+}
+
+TMap<int32, TArray<FChannelMapping>> FResourceMappingList::GetChannelMappingPerFrame(int32 InputIndex, EResourceName Name) const
+{
+	check(InputIndex >= 0 && InputIndex < InputMapping.Num());
+
+	return InputMapping[InputIndex].GetChannelMappingPerFrame(Name);
+}
+
+template<>
+FResourceMappingList MakeTensorLayout<FNNEDenoiserInputMappingData>(UDataTable* DataTable)
+{
+	TMap<int32, TMap<int32, FResourceInfo>> Map;
+	DataTable->ForeachRow<FNNEDenoiserInputMappingData>("FResourceLayout", [&] (const FName &Key, const FNNEDenoiserInputMappingData &Value)
+	{
+		Detail::Add(Map, Value.TensorIndex, Value.TensorChannel, ToResourceName(Value.Resource), Value.ResourceChannel, 0);
+	});
+
+	return Detail::MakeTensorLayout(Map);
+}
+
+template<>
+FResourceMappingList MakeTensorLayout<FNNEDenoiserOutputMappingData>(UDataTable* DataTable)
+{
+	TMap<int32, TMap<int32, FResourceInfo>> Map;
+	DataTable->ForeachRow<FNNEDenoiserOutputMappingData>("FResourceLayout", [&] (const FName &Key, const FNNEDenoiserOutputMappingData &Value)
+	{
+		Detail::Add(Map, Value.TensorIndex, Value.TensorChannel, ToResourceName(Value.Resource), Value.ResourceChannel, 0);
+	});
+
+	return Detail::MakeTensorLayout(Map);
+}
+
+template<>
+FResourceMappingList MakeTensorLayout<FNNEDenoiserTemporalInputMappingData>(UDataTable* DataTable)
+{
+	TMap<int32, TMap<int32, FResourceInfo>> Map;
+	DataTable->ForeachRow<FNNEDenoiserTemporalInputMappingData>("FResourceLayout", [&] (const FName &Key, const FNNEDenoiserTemporalInputMappingData &Value)
+	{
+		Detail::Add(Map, Value.TensorIndex, Value.TensorChannel, ToResourceName(Value.Resource), Value.ResourceChannel, Value.FrameIndex);
+	});
+
+	return Detail::MakeTensorLayout(Map);
+}
+
+template<>
+FResourceMappingList MakeTensorLayout<FNNEDenoiserTemporalOutputMappingData>(UDataTable* DataTable)
+{
+	TMap<int32, TMap<int32, FResourceInfo>> Map;
+	DataTable->ForeachRow<FNNEDenoiserTemporalOutputMappingData>("FResourceLayout", [&] (const FName &Key, const FNNEDenoiserTemporalOutputMappingData &Value)
+	{
+		Detail::Add(Map, Value.TensorIndex, Value.TensorChannel, ToResourceName(Value.Resource), Value.ResourceChannel, 0);
+	});
+
+	return Detail::MakeTensorLayout(Map);
+}
+
+} // namespace UE::NNEDenoiser::Private
\ No newline at end of file
diff --git a/Engine/Plugins/Experimental/NNEDenoiser/Source/NNEDenoiser/Private/NNEDenoiserResourceMapping.h b/Engine/Plugins/Experimental/NNEDenoiser/Source/NNEDenoiser/Private/NNEDenoiserResourceMapping.h
--- a/Engine/Plugins/Experimental/NNEDenoiser/Source/NNEDenoiser/Private/NNEDenoiserResourceMapping.h
+++ b/Engine/Plugins/Experimental/NNEDenoiser/Source/NNEDenoiser/Private/NNEDenoiserResourceMapping.h
@@ -4,11 +4,12 @@
 
 #include "Containers/Array.h"
 #include "Containers/Map.h"
-#include "NNEDenoiserModelIOMappingData.h"
+#include "NNEDenoiserIOMappingData.h"
+#include "NNEDenoiserResourceName.h"
 
 namespace UE::NNEDenoiser::Private
 {
-	
+
 struct FResourceInfo
 {
 	EResourceName Name;
@@ -16,163 +17,81 @@ struct FResourceInfo
 	int32 Frame;
 };
 
+struct FChannelMapping
+{
+	int32 TensorChannel;
+	int32 ResourceChannel;
+};
+
+// Resource mapping for one tensor:
+// 0: ResourceInfo0 -> tensor channel 0 maps to channel ResourceInfo0.Channel of resource ResourceInfo0.Name from frame ResourceInfo0.Frame
+// 1: ...
 class FResourceMapping
 {
 public:
-	void Add(FResourceInfo Info)
-	{
-		ChannelMapping.Add(MoveTemp(Info));
-	}
-
-	FResourceInfo& Add_GetRef(FResourceInfo Info)
-	{
-		return ChannelMapping.Add_GetRef(MoveTemp(Info));
-	}
-
-	const FResourceInfo& GetChecked(int32 Channel) const
-	{
-		check(Channel >= 0 && Channel < ChannelMapping.Num());
-
-		return ChannelMapping[Channel];
-	}
-
-	int32 Num() const
-	{
-		return ChannelMapping.Num();
-	}
-
-	bool HasResource(EResourceName Name) const
-	{
-		for (const auto& Info: ChannelMapping)
-		{
-			if (Info.Name == Name)
-			{
-				return true;
-			}
-		}
-		return false;
-	}
-
-	int32 NumFrames(EResourceName Name) const
-	{
-		int32 MinFrameIndex = 1;
-		for (const auto& Info: ChannelMapping)
-		{
-			if (Info.Name == Name)
-			{
-				MinFrameIndex = FMath::Min(MinFrameIndex, Info.Frame);
-			}
-		}
-
-		return 1 - MinFrameIndex;
-	}
-
-	TMap<int32, TArray<FIntPoint>> GetChannelMappingPerFrame(EResourceName Name) const
-	{
-		TMap<int32, TArray<FIntPoint>> Result;
-		for (int32 I = 0; I < ChannelMapping.Num(); I++)
-		{
-			const FResourceInfo& Info = ChannelMapping[I];
-
-			if (Info.Name == Name)
-			{
-				Result.FindOrAdd(Info.Frame).Emplace(I, Info.Channel);
-			}
-		}
-
-		return Result;
-	}
+	// Add mapping from tensor channel Num()-1 to resource defined by 'Info'
+	void Add(FResourceInfo Info);
+
+	// Add mapping from tensor channel Num()-1 to resource defined by 'Info' and return reference to it
+	FResourceInfo& Add_GetRef(FResourceInfo Info);
+
+	const FResourceInfo& GetChecked(int32 Channel) const;
+
+	// Number of mapped tensor channels
+	int32 Num() const;
+
+	// Number of past frames affected by the mapping including current frame for resource 'Name'
+	int32 NumFrames(EResourceName Name) const;
+
+	// Does mapping use resource 'Name'
+	bool HasResource(EResourceName Name) const;
+
+	// Get mapping from tensor channel to resource channel per frame for resource 'Name'
+	TMap<int32, TArray<FChannelMapping>> GetChannelMappingPerFrame(EResourceName Name) const;
 
 private:
 	TArray<FResourceInfo> ChannelMapping;
 };
 
+// Resource mapping list for multiple tensors:
+// 0: ResourceMapping0 -> tensor 0 maps to resources as defined by ResourceMapping0
+// 1: ...
 class FResourceMappingList
 {
 public:
-	void Add(FResourceMapping ResourceMapping)
-	{
-		InputMapping.Add(MoveTemp(ResourceMapping));
-	}
-
-	FResourceMapping& Add_GetRef(FResourceMapping ResourceMapping)
-	{
-		return InputMapping.Add_GetRef(MoveTemp(ResourceMapping));
-	}
-
-	const FResourceMapping* Get(int32 InputIndex) const
-	{
-		if (InputIndex >= 0 && InputIndex < InputMapping.Num())
-		{
-			return &InputMapping[InputIndex];
-		}
-
-		return {};
-	}
-
-	const FResourceMapping& GetChecked(int32 InputIndex) const
-	{
-		check(InputIndex >= 0 && InputIndex < InputMapping.Num());
-
-		return InputMapping[InputIndex];
-	}
-
-	int32 Num() const
-	{
-		return InputMapping.Num();
-	}
-
-	int32 NumChannels(int32 InputIndex) const
-	{
-		check(InputIndex >= 0 && InputIndex < InputMapping.Num());
-		
-		return InputMapping[InputIndex].Num();
-	}
-
-	int32 NumFrames(EResourceName Name) const
-	{
-		int32 MaxNumFrames = 0;
-		for (const auto& Input : InputMapping)
-		{
-			const int32 NumFrames = Input.NumFrames(Name);
-			if (MaxNumFrames < NumFrames)
-			{
-				MaxNumFrames = NumFrames;
-			}
-		}
-
-		return MaxNumFrames;
-	}
-
-	bool HasResource(EResourceName Name) const
-	{
-		for (const auto& Input : InputMapping)
-		{
-			if (Input.HasResource(Name))
-			{
-				return true;
-			}
-		}
-
-		return false;
-	}
-
-	bool HasResource(int32 InputIndex, EResourceName Name) const
-	{
-		check(InputIndex >= 0 && InputIndex < InputMapping.Num());
-
-		return InputMapping[InputIndex].HasResource(Name);
-	}
-
-	TMap<int32, TArray<FIntPoint>> GetChannelMappingPerFrame(int32 InputIndex, EResourceName Name) const
-	{
-		check(InputIndex >= 0 && InputIndex < InputMapping.Num());
-
-		return InputMapping[InputIndex].GetChannelMappingPerFrame(Name);
-	}
+	// Add resource mapping for tensor Num()-1
+	void Add(FResourceMapping ResourceMapping);
+
+	// Add resource mapping for tensor Num()-1 and return reference to it
+	FResourceMapping& Add_GetRef(FResourceMapping ResourceMapping);
+
+	const FResourceMapping* Get(int32 InputIndex) const;
+	const FResourceMapping& GetChecked(int32 InputIndex) const;
+
+	// Number of mapped tensors
+	int32 Num() const;
+
+	// Number of mapped tensor channels for input tensor 'InputIndex'
+	int32 NumChannels(int32 InputIndex) const;
+
+	// Number of past frames affected by the mapping including current frame for resource 'Name'
+	int32 NumFrames(EResourceName Name) const;
+
+	// Does mapping use resource 'Name'
+	bool HasResource(EResourceName Name) const;
+
+	// Does mapping for tensor 'InputIndex' use resource 'Name'
+	bool HasResource(int32 InputIndex, EResourceName Name) const;
+
+	// Get mapping from tensor channel to resource channel per frame for resource 'Name' and tensor 'InputIndex"
+	TMap<int32, TArray<FChannelMapping>> GetChannelMappingPerFrame(int32 InputIndex, EResourceName Name) const;
 
 private:
 	TArray<FResourceMapping> InputMapping;
 };
 
+// Helper method to make resource mapping list from data table asset using row struct 'RowStructType'
+template<class RowStructType>
+FResourceMappingList MakeTensorLayout(UDataTable* DataTable);
+
 } // namespace UE::NNEDenoiser::Private
\ No newline at end of file
diff --git a/Engine/Plugins/Experimental/NNEDenoiser/Source/NNEDenoiser/Private/NNEDenoiserSettings.cpp b/Engine/Plugins/Experimental/NNEDenoiser/Source/NNEDenoiser/Private/NNEDenoiserSettings.cpp
--- a/Engine/Plugins/Experimental/NNEDenoiser/Source/NNEDenoiser/Private/NNEDenoiserSettings.cpp
+++ b/Engine/Plugins/Experimental/NNEDenoiser/Source/NNEDenoiser/Private/NNEDenoiserSettings.cpp
@@ -12,6 +12,7 @@ UNNEDenoiserSettings::UNNEDenoiserSettings()
 	SectionName = TEXT("NNE Denoiser");
 
 	DenoiserAsset = FSoftObjectPath(TEXT("/NNEDenoiser/NNED_Oidn2-3_Balanced.NNED_Oidn2-3_Balanced"));
+	TemporalDenoiserAsset = FSoftObjectPath(TEXT("/NNEDenoiser/NNEDT_Oidn2-3_Balanced.NNEDT_Oidn2-3_Balanced"));
 }
 
 void UNNEDenoiserSettings::PostInitProperties()
diff --git a/Engine/Plugins/Experimental/NNEDenoiser/Source/NNEDenoiser/Private/NNEDenoiserUtils.h b/Engine/Plugins/Experimental/NNEDenoiser/Source/NNEDenoiser/Private/NNEDenoiserUtils.h
--- a/Engine/Plugins/Experimental/NNEDenoiser/Source/NNEDenoiser/Private/NNEDenoiserUtils.h
+++ b/Engine/Plugins/Experimental/NNEDenoiser/Source/NNEDenoiser/Private/NNEDenoiserUtils.h
@@ -3,7 +3,6 @@
 #pragma once
 
 #include "NNEDenoiserLog.h"
-#include "NNEDenoiserModelIOMappingData.h"
 #include "NNEDenoiserParameters.h"
 #include "NNEDenoiserShadersMappedCopyCS.h"
 #include "NNETypes.h"
diff --git a/Engine/Plugins/Experimental/NNEDenoiser/Source/NNEDenoiser/Private/NNEDenoiserViewExtension.cpp b/Engine/Plugins/Experimental/NNEDenoiser/Source/NNEDenoiser/Private/NNEDenoiserViewExtension.cpp
--- a/Engine/Plugins/Experimental/NNEDenoiser/Source/NNEDenoiser/Private/NNEDenoiserViewExtension.cpp
+++ b/Engine/Plugins/Experimental/NNEDenoiser/Source/NNEDenoiser/Private/NNEDenoiserViewExtension.cpp
@@ -2,17 +2,19 @@
 
 #include "NNEDenoiserViewExtension.h"
 #include "NNE.h"
+#include "NNEDenoiserAsset.h"
 #include "NNEDenoiserAutoExposure.h"
 #include "NNEDenoiserGenericDenoiser.h"
+#include "NNEDenoiserIOMappingData.h"
 #include "NNEDenoiserIOProcessBase.h"
 #include "NNEDenoiserLog.h"
-#include "NNEDenoiserAsset.h"
 #include "NNEDenoiserModelInstanceCPU.h"
 #include "NNEDenoiserModelInstanceGPU.h"
 #include "NNEDenoiserModelInstanceRDG.h"
-#include "NNEDenoiserModelIOMappingData.h"
 #include "NNEDenoiserPathTracingDenoiser.h"
 #include "NNEDenoiserPathTracingSpatialTemporalDenoiser.h"
+#include "NNEDenoiserResourceMapping.h"
+#include "NNEDenoiserTemporalAsset.h"
 #include "NNEDenoiserTransferFunctionOidn.h"
 #include "NNEDenoiserUtils.h"
 #include "NNEModelData.h"
@@ -30,7 +32,7 @@ static TAutoConsoleVariable<bool> CVarNNEDenoiser(
 static TAutoConsoleVariable<int32> CVarNNEDenoiserAsset(
 	TEXT("NNEDenoiser.Asset"),
 	0,
-	TEXT("Defines asset to used to create the denoiser.\n")
+	TEXT("Defines denoiser asset to used to create the denoiser.\n")
 	TEXT("  0: Use denoiser asset defined by Project Settings\n")
 	TEXT("  1: OIDN Fast\n")
 	TEXT("  2: OIDN Balanced\n")
@@ -40,6 +42,15 @@ static TAutoConsoleVariable<int32> CVarNNEDenoiserAsset(
 	TEXT("  6: OIDN High Quality | Alpha")
 );
 
+static TAutoConsoleVariable<int32> CVarNNEDenoiserTemporalAsset(
+	TEXT("NNEDenoiser.TemporalAsset"),
+	0,
+	TEXT("Defines temporal denoiser asset to used to create the temporal denoiser.\n")
+	TEXT("  0: Use temporal denoiser asset defined by Project Settings\n")
+	TEXT("  1: OIDN Balanced\n")
+	TEXT("  2: OIDN Balanced | Alpha")
+);
+
 static TAutoConsoleVariable<int32> CVarNNEDenoiserRuntimeType(
 	TEXT("NNEDenoiser.Runtime.Type"),
 	2,
@@ -83,6 +94,22 @@ FString GetDenoiserAssetNameFromCVarAndSettings(const UNNEDenoiserSettings* Sett
 	return FString();
 }
 
+FString GetDenoiserTemporalAssetNameFromCVarAndSettings(const UNNEDenoiserSettings* Settings)
+{
+	const int32 Idx = FMath::Clamp(CVarNNEDenoiserTemporalAsset.GetValueOnGameThread(), 0, 6);
+	switch(Idx)
+	{
+		case 0: return !Settings->TemporalDenoiserAsset.IsNull() ? Settings->TemporalDenoiserAsset.ToString() : FString();
+
+		case 1: return TEXT("/NNEDenoiser/NNEDT_Oidn2-3_Balanced.NNEDT_Oidn2-3_Balanced");
+
+		// Alpha
+		case 2: return TEXT("/NNEDenoiser/NNEDT_Oidn2-3_Balanced_Alpha.NNEDT_Oidn2-3_Balanced_Alpha");
+	}
+	check(false);
+	return FString();
+}
+
 TUniquePtr<FGenericDenoiser> CreateNNEDenoiser(
 	UNNEModelData& ModelData,
 	EDenoiserRuntimeType RuntimeType,
@@ -169,45 +196,8 @@ TUniquePtr<FGenericDenoiser> CreateNNEDenoiser(
 		TransferFunction);
 }
 
-static FResourceMappingList MakeTensorLayout(UDataTable* DataTable)
-{
-	TMap<int32, TMap<int32, FResourceInfo>> Map;
-	DataTable->ForeachRow<FNNEDenoiserModelIOMappingData>("FResourceLayout", [&] (const FName &Key, const FNNEDenoiserModelIOMappingData &Value)
-	{
-		if (Value.TensorChannel < 0)
-		{
-			for (int32 I = 0; I < -Value.TensorChannel; I++)
-			{
-				Map.FindOrAdd(Value.TensorIndex).FindOrAdd(I) = FResourceInfo{Value.Resource, I, Value.FrameIndex};
-			}
-		}
-		else
-		{
-			Map.FindOrAdd(Value.TensorIndex).FindOrAdd(Value.TensorChannel) = FResourceInfo{Value.Resource, Value.ResourceChannel, Value.FrameIndex};
-		}
-	});
-
-	FResourceMappingList Result{};
-
-	for (int32 I = 0; I < Map.Num(); I++)
-	{
-		checkf(Map.Contains(I), TEXT("Missing intput/output %d, must be continuous!"), I);
-
-		const TMap<int32, FResourceInfo>& InnerMap = Map[I];
-		FResourceMapping& Mapping = Result.Add_GetRef({});
-
-		for (int32 J = 0; J < InnerMap.Num(); J++)
-		{
-			checkf(InnerMap.Contains(J), TEXT("Missing tensor info for channel %d, must be continuous!"), J);
-
-			Mapping.Add(InnerMap[J]);
-		}
-	}
-
-	return Result;
-}
-
-FParameters GetParametersValidated(const UNNEDenoiserAsset& DenoiserAsset)
+template<class AssetType>
+FParameters GetParametersValidated(const AssetType& DenoiserAsset)
 {
 	FParameters Parameters{
 		.TilingConfig =
@@ -258,6 +248,7 @@ FParameters GetParametersValidated(const UNNEDenoiserAsset& DenoiserAsset)
 	return Parameters;
 }
 
+template<class AssetType, class InputMappingType, class OutputMappingType>
 TUniquePtr<FGenericDenoiser> CreateNNEDenoiserFromAsset(const FString& AssetName, EDenoiserRuntimeType RuntimeType, const FString& RuntimeNameOverride)
 {
 	if (AssetName.IsEmpty())
@@ -266,7 +257,7 @@ TUniquePtr<FGenericDenoiser> CreateNNEDenoiserFromAsset(const FString& AssetName
 		return {};
 	}
 
-	UNNEDenoiserAsset *DenoiserAsset = LoadObject<UNNEDenoiserAsset>(nullptr, *AssetName);
+	AssetType *DenoiserAsset = LoadObject<AssetType>(nullptr, *AssetName);
 	if (!DenoiserAsset)
 	{
 		UE_LOG(LogNNEDenoiser, Error, TEXT("Could not load denoiser model data asset!"));
@@ -293,23 +284,18 @@ TUniquePtr<FGenericDenoiser> CreateNNEDenoiserFromAsset(const FString& AssetName
 	FResourceMappingList InputLayout;
 	if (InputMappingTable)
 	{
-		InputLayout = MakeTensorLayout(InputMappingTable);
+		InputLayout = MakeTensorLayout<InputMappingType>(InputMappingTable);
 	}
 
 	FResourceMappingList OutputLayout;
 	if (OutputMappingTable)
 	{
-		OutputLayout = MakeTensorLayout(OutputMappingTable);
+		OutputLayout = MakeTensorLayout<OutputMappingType>(OutputMappingTable);
 	}
 
-	const bool bIsOidnModel = AssetName.Contains(TEXT("oidn2"));
-	const bool bIsInHouseModel = AssetName.Contains(TEXT("kpcn"));
-
-	check(!bIsOidnModel || !bIsInHouseModel);
-
 	TUniquePtr<FAutoExposure> AutoExposure;
 	TSharedPtr<ITransferFunction> TransferFunction;
-	if (bIsOidnModel)
+	if (AssetName.Contains(TEXT("oidn2")))
 	{
 		AutoExposure = MakeUnique<FAutoExposure>();
 		TransferFunction = MakeShared<Oidn::FTransferFunction>();
@@ -347,8 +333,7 @@ FViewExtension::~FViewExtension()
 	}
 #endif
 
-	UnregisterDenoiser(TEXT("NNE_OIDN"));
-	UnregisterDenoiser(TEXT("NNE_KPCN"));
+	UnregisterDenoiser(TEXT("NNEDenoiser"));
 }
 
 void FViewExtension::BeginRenderViewFamily(FSceneViewFamily& InViewFamily)
@@ -360,23 +345,23 @@ void FViewExtension::PreRenderViewFamily_RenderThread(FRDGBuilder& GraphBuilder,
 {
 	if (!bDenoiserEnabled)
 	{
-		UnregisterDenoiser(TEXT("NNE_OIDN"));
-		UnregisterDenoiser(TEXT("NNE_KPCN"));
+		UnregisterDenoiser(TEXT("NNEDenoiser"));
 	}
 	else
 	{
+		if (DenoiserToSwap.IsValid() || SpatialTemporalDenoiserToSwap.IsValid())
+		{
+			UnregisterDenoiser(TEXT("NNEDenoiser"));
+		}
+
 		if (DenoiserToSwap.IsValid())
 		{
-			UnregisterDenoiser(TEXT("NNE_OIDN"));
-			UnregisterDenoiser(TEXT("NNE_KPCN"));
-			RegisterSpatialDenoiser(MoveTemp(DenoiserToSwap),TEXT("NNE_OIDN"));
+			RegisterSpatialDenoiser(MoveTemp(DenoiserToSwap),TEXT("NNEDenoiser"));
 		}
 
 		if (SpatialTemporalDenoiserToSwap.IsValid())
 		{
-			UnregisterDenoiser(TEXT("NNE_OIDN"));
-			UnregisterDenoiser(TEXT("NNE_KPCN"));
-			RegisterSpatialTemporalDenoiser(MoveTemp(SpatialTemporalDenoiserToSwap), TEXT("NNE_KPCN"));
+			RegisterSpatialTemporalDenoiser(MoveTemp(SpatialTemporalDenoiserToSwap), TEXT("NNEDenoiser"));
 		}
 	}
 }
@@ -388,7 +373,8 @@ void FViewExtension::ApplySettings(const UNNEDenoiserSettings* Settings)
 	bNeedsUpdate |= CVarNNEDenoiser.GetValueOnGameThread() != bDenoiserEnabled;
 	bNeedsUpdate |= GetDenoiserRuntimeTypeFromCVar() != RuntimeType;
 	bNeedsUpdate |= CVarNNEDenoiserRuntimeName.GetValueOnGameThread() != RuntimeName;
-	bNeedsUpdate |= GetDenoiserAssetNameFromCVarAndSettings(Settings) != ModelDataName;
+	bNeedsUpdate |= GetDenoiserAssetNameFromCVarAndSettings(Settings) != AssetName;
+	bNeedsUpdate |= GetDenoiserTemporalAssetNameFromCVarAndSettings(Settings) != TemporalAssetName;
 
 	if (!bNeedsUpdate)
 	{
@@ -406,25 +392,33 @@ void FViewExtension::ApplySettings(const UNNEDenoiserSettings* Settings)
 
 	RuntimeType = GetDenoiserRuntimeTypeFromCVar();
 	RuntimeName = CVarNNEDenoiserRuntimeName.GetValueOnGameThread();
-	ModelDataName = GetDenoiserAssetNameFromCVarAndSettings(Settings);
+	AssetName = GetDenoiserAssetNameFromCVarAndSettings(Settings);
+	TemporalAssetName = GetDenoiserTemporalAssetNameFromCVarAndSettings(Settings);
 
-	UE_LOG(LogNNEDenoiser, Log, TEXT("Create denoiser from asset %s..."), *ModelDataName);
+	// UE_LOG(LogNNEDenoiser, Log, TEXT("Create denoiser from asset %s..."), *AssetName);
 
-	TUniquePtr<FGenericDenoiser> Denoiser = CreateNNEDenoiserFromAsset(ModelDataName, RuntimeType, RuntimeName);
+	TUniquePtr<FGenericDenoiser> Denoiser = CreateNNEDenoiserFromAsset<UNNEDenoiserAsset, FNNEDenoiserInputMappingData, FNNEDenoiserOutputMappingData>(AssetName, RuntimeType, RuntimeName);
 	if (Denoiser.IsValid())
 	{
-		if (ModelDataName.Contains(TEXT("kpcn")))
-		{
-			SpatialTemporalDenoiserToSwap = MakeUnique<FPathTracingSpatialTemporalDenoiser>(MoveTemp(Denoiser));
-		}
-		else
-		{
-			DenoiserToSwap = MakeUnique<FPathTracingDenoiser>(MoveTemp(Denoiser));
-		}
+		UE_LOG(LogNNEDenoiser, Log, TEXT("Create denoiser from asset %s..."), *AssetName);
+
+		DenoiserToSwap = MakeUnique<FPathTracingDenoiser>(MoveTemp(Denoiser));
 	}
 	else
 	{
-		UE_LOG(LogNNEDenoiser, Error, TEXT("Could not create denoiser!"));
+		UE_LOG(LogNNEDenoiser, Error, TEXT("Could not create denoiser from asset %d!"), *AssetName);
+	}
+
+	TUniquePtr<FGenericDenoiser> TemporalDenoiser = CreateNNEDenoiserFromAsset<UNNEDenoiserTemporalAsset, FNNEDenoiserTemporalInputMappingData, FNNEDenoiserTemporalOutputMappingData>(TemporalAssetName, RuntimeType, RuntimeName);
+	if (TemporalDenoiser.IsValid())
+	{
+		UE_LOG(LogNNEDenoiser, Log, TEXT("Create temporal denoiser from asset %s..."), *TemporalAssetName);
+
+		SpatialTemporalDenoiserToSwap = MakeUnique<FPathTracingSpatialTemporalDenoiser>(MoveTemp(TemporalDenoiser));
+	}
+	else
+	{
+		UE_LOG(LogNNEDenoiser, Error, TEXT("Could not create temporal denoiser from asset %d!"), *TemporalAssetName);
 	}
 }
 
diff --git a/Engine/Plugins/Experimental/NNEDenoiser/Source/NNEDenoiser/Private/NNEDenoiserViewExtension.h b/Engine/Plugins/Experimental/NNEDenoiser/Source/NNEDenoiser/Private/NNEDenoiserViewExtension.h
--- a/Engine/Plugins/Experimental/NNEDenoiser/Source/NNEDenoiser/Private/NNEDenoiserViewExtension.h
+++ b/Engine/Plugins/Experimental/NNEDenoiser/Source/NNEDenoiser/Private/NNEDenoiserViewExtension.h
@@ -41,7 +41,8 @@ class FViewExtension final : public FSceneViewExtensionBase
 		bool bDenoiserEnabled = true;
 		EDenoiserRuntimeType RuntimeType;
 		FString RuntimeName;
-		FString ModelDataName;
+		FString AssetName;
+		FString TemporalAssetName;
 	};
 
 } // namespace UE::NNEDenoiser::Private
\ No newline at end of file
diff --git a/Engine/Plugins/Experimental/NNEDenoiser/Source/NNEDenoiser/Public/NNEDenoiserAsset.h b/Engine/Plugins/Experimental/NNEDenoiser/Source/NNEDenoiser/Public/NNEDenoiserAsset.h
--- a/Engine/Plugins/Experimental/NNEDenoiser/Source/NNEDenoiser/Public/NNEDenoiserAsset.h
+++ b/Engine/Plugins/Experimental/NNEDenoiser/Source/NNEDenoiser/Public/NNEDenoiserAsset.h
@@ -4,33 +4,11 @@
 
 #include "Engine/DataAsset.h"
 #include "Engine/DataTable.h"
+#include "NNEDenoiserTilingConfig.h"
 #include "NNEModelData.h"
 
 #include "NNEDenoiserAsset.generated.h"
 
-/** Tiling configuration for fixed and dynamic size models */
-USTRUCT(BlueprintType)
-struct FTilingConfig
-{
-	GENERATED_BODY()
-
-	/** Tile size alignment (applies only to dynamic size models) */
-	UPROPERTY(EditAnywhere, BlueprintReadWrite, Category=NNEDenoiser, meta = (DisplayName = "Size Alignment"))
-	int32 Alignment = 1;
-
-	/** Tile overlap */
-	UPROPERTY(EditAnywhere, BlueprintReadWrite, Category=NNEDenoiser)
-	int32 Overlap = 0;
-
-	/** Maximum tile size (applies only to dynamic size models) */
-	UPROPERTY(EditAnywhere, BlueprintReadWrite, Category=NNEDenoiser)
-	int32 MaxSize = 0;
-
-	/** Minimum tile size (applies only to dynamic size models) */
-	UPROPERTY(EditAnywhere, BlueprintReadWrite, Category=NNEDenoiser)
-	int32 MinSize = 1;
-};
-
 /** Denoiser model data asset */
 UCLASS(BlueprintType)
 class NNEDENOISER_API UNNEDenoiserAsset : public UDataAsset
@@ -43,11 +21,11 @@ public:
 	TSoftObjectPtr<UNNEModelData> ModelData;
 
 	/** Input mapping table */
-	UPROPERTY(EditAnywhere, BlueprintReadWrite, Category=NNEDenoiser, meta = (RequiredAssetDataTags = "RowStructure=/Script/NNEDenoiser.NNEDenoiserModelIOMappingData"))
+	UPROPERTY(EditAnywhere, BlueprintReadWrite, Category=NNEDenoiser, meta = (RequiredAssetDataTags = "RowStructure=/Script/NNEDenoiser.NNEDenoiserInputMappingData"))
 	TSoftObjectPtr<UDataTable> InputMapping;
 
 	/** Output mapping table */
-	UPROPERTY(EditAnywhere, BlueprintReadWrite, Category=NNEDenoiser, meta = (RequiredAssetDataTags = "RowStructure=/Script/NNEDenoiser.NNEDenoiserModelIOMappingData"))
+	UPROPERTY(EditAnywhere, BlueprintReadWrite, Category=NNEDenoiser, meta = (RequiredAssetDataTags = "RowStructure=/Script/NNEDenoiser.NNEDenoiserOutputMappingData"))
 	TSoftObjectPtr<UDataTable> OutputMapping;
 
 	/** Tiling configuration */
diff --git a/Engine/Plugins/Experimental/NNEDenoiser/Source/NNEDenoiser/Public/NNEDenoiserIOMappingData.h b/Engine/Plugins/Experimental/NNEDenoiser/Source/NNEDenoiser/Public/NNEDenoiserIOMappingData.h
new file mode 100644
--- /dev/null
+++ b/Engine/Plugins/Experimental/NNEDenoiser/Source/NNEDenoiser/Public/NNEDenoiserIOMappingData.h
@@ -0,0 +1,122 @@
+// Copyright Epic Games, Inc. All Rights Reserved.
+
+#pragma once
+
+#include "Engine/DataTable.h"
+#include "NNEDenoiserResourceName.h"
+
+#include "NNEDenoiserIOMappingData.generated.h"
+
+/** An enum to represent resource names used for denoiser input mapping */
+UENUM()
+enum EInputResourceName : uint8
+{
+	I_Color UMETA(DisplayName="Color"),
+	I_Albedo UMETA(DisplayName="Albedo"),
+	I_Normal UMETA(DisplayName="Normal"),
+	I_Output UMETA(DisplayName="Output")
+};
+
+/** An enum to represent resource names used for denoiser output mapping */
+UENUM()
+enum EOutputResourceName : uint8
+{
+	O_Output UMETA(DisplayName="Output")
+};
+
+/** An enum to represent resource names used for temporal denoiser input mapping */
+UENUM()
+enum ETemporalInputResourceName : uint8
+{
+	TI_Color UMETA(DisplayName="Color"),
+	TI_Albedo UMETA(DisplayName="Albedo"),
+	TI_Normal UMETA(DisplayName="Normal"),
+	TI_Flow UMETA(DisplayName="Flow"),
+	TI_Output UMETA(DisplayName="Output")
+};
+
+/** An enum to represent resource names used for temporal denoiser output mapping */
+UENUM()
+enum ETemporalOutputResourceName : uint8
+{
+	TO_Output UMETA(DisplayName="Output")
+};
+
+/** Table row base for denoiser basic input and output mapping */
+USTRUCT(BlueprintType)
+struct FNNEDenoiserBaseMappingData : public FTableRowBase
+{
+	GENERATED_USTRUCT_BODY()
+
+public:
+	/** Input/output tensor index */
+	UPROPERTY(EditAnywhere, BlueprintReadWrite, Category=NNEDenoiser)
+	int32 TensorIndex = 0;
+
+	/** Input/output tensor channel */
+	UPROPERTY(EditAnywhere, BlueprintReadWrite, Category=NNEDenoiser)
+	int32 TensorChannel = 0;
+
+	/** Resource channel */
+	UPROPERTY(EditAnywhere, BlueprintReadWrite, Category=NNEDenoiser)
+	int32 ResourceChannel = 0;
+};
+
+/** Table row base for denoiser input mapping */
+USTRUCT(BlueprintType)
+struct FNNEDenoiserInputMappingData : public FNNEDenoiserBaseMappingData
+{
+	GENERATED_USTRUCT_BODY()
+
+	/** Mapped resource name */
+	UPROPERTY(EditAnywhere, BlueprintReadWrite, Category=NNEDenoiser)
+	TEnumAsByte<EInputResourceName> Resource = EInputResourceName::I_Color;
+};
+
+/** Table row base for denoiser output mapping */
+USTRUCT(BlueprintType)
+struct FNNEDenoiserOutputMappingData : public FNNEDenoiserBaseMappingData
+{
+	GENERATED_USTRUCT_BODY()
+
+	/** Mapped resource name */
+	UPROPERTY(EditAnywhere, BlueprintReadWrite, Category=NNEDenoiser)
+	TEnumAsByte<EOutputResourceName> Resource = EOutputResourceName::O_Output;
+};
+
+/** Table row base for temporal denoiser input mapping */
+USTRUCT(BlueprintType)
+struct FNNEDenoiserTemporalInputMappingData : public FNNEDenoiserBaseMappingData
+{
+	GENERATED_USTRUCT_BODY()
+
+public:
+	/** Mapped resource name */
+	UPROPERTY(EditAnywhere, BlueprintReadWrite, Category=NNEDenoiser)
+	TEnumAsByte<ETemporalInputResourceName> Resource = ETemporalInputResourceName::TI_Color;
+
+	/** Resource frame index */
+	UPROPERTY(EditAnywhere, BlueprintReadWrite, Category=NNEDenoiser)
+	int32 FrameIndex = 0;
+};
+
+/** Table row base for temporal denoiser output mapping */
+USTRUCT(BlueprintType)
+struct FNNEDenoiserTemporalOutputMappingData : public FNNEDenoiserBaseMappingData
+{
+	GENERATED_USTRUCT_BODY()
+
+	/** Mapped resource name */
+	UPROPERTY(EditAnywhere, BlueprintReadWrite, Category=NNEDenoiser)
+	TEnumAsByte<ETemporalOutputResourceName> Resource = ETemporalOutputResourceName::TO_Output;
+};
+
+namespace UE::NNEDenoiser
+{
+
+EResourceName ToResourceName(EInputResourceName Name);
+EResourceName ToResourceName(EOutputResourceName Name);
+EResourceName ToResourceName(ETemporalInputResourceName Name);
+EResourceName ToResourceName(ETemporalOutputResourceName Name);
+
+} // namespace UE::NNEDenoiser
\ No newline at end of file
diff --git a/Engine/Plugins/Experimental/NNEDenoiser/Source/NNEDenoiser/Public/NNEDenoiserModelIOMappingData.h b/Engine/Plugins/Experimental/NNEDenoiser/Source/NNEDenoiser/Public/NNEDenoiserModelIOMappingData.h
deleted file mode 100644
diff --git a/Engine/Plugins/Experimental/NNEDenoiser/Source/NNEDenoiser/Public/NNEDenoiserResourceName.h b/Engine/Plugins/Experimental/NNEDenoiser/Source/NNEDenoiser/Public/NNEDenoiserResourceName.h
new file mode 100644
--- /dev/null
+++ b/Engine/Plugins/Experimental/NNEDenoiser/Source/NNEDenoiser/Public/NNEDenoiserResourceName.h
@@ -0,0 +1,16 @@
+// Copyright Epic Games, Inc. All Rights Reserved.
+
+#pragma once
+
+#include "NNEDenoiserResourceName.generated.h"
+
+/** An enum to represent resource names used for input and output mapping */
+UENUM()
+enum EResourceName : uint8
+{
+	Color,
+	Albedo,
+	Normal,
+	Flow,
+	Output
+};
\ No newline at end of file
diff --git a/Engine/Plugins/Experimental/NNEDenoiser/Source/NNEDenoiser/Public/NNEDenoiserSettings.h b/Engine/Plugins/Experimental/NNEDenoiser/Source/NNEDenoiser/Public/NNEDenoiserSettings.h
--- a/Engine/Plugins/Experimental/NNEDenoiser/Source/NNEDenoiser/Public/NNEDenoiserSettings.h
+++ b/Engine/Plugins/Experimental/NNEDenoiser/Source/NNEDenoiser/Public/NNEDenoiserSettings.h
@@ -4,6 +4,7 @@
 
 #include "Engine/DeveloperSettingsBackedByCVars.h"
 #include "NNEDenoiserAsset.h"
+#include "NNEDenoiserTemporalAsset.h"
 
 #include "NNEDenoiserSettings.generated.h"
 
@@ -31,6 +32,10 @@ public:
 	UPROPERTY(Config, EditAnywhere, Category = "NNE Denoiser", meta = (DisplayName = "Denoiser Asset", ToolTip = "Select the denoiser asset"))
 	TSoftObjectPtr<UNNEDenoiserAsset> DenoiserAsset;
 
+	/** Temporal denoiser asset data used to create a NNE Denoiser */
+	UPROPERTY(Config, EditAnywhere, Category = "NNE Denoiser", meta = (DisplayName = "Temporal Denoiser Asset", ToolTip = "Select the temporal denoiser asset"))
+	TSoftObjectPtr<UNNEDenoiserTemporalAsset> TemporalDenoiserAsset;
+
 private:
 	/** Runtime type used to run the NNE Denoiser model. Backed by the console variable 'NNEDenoiser.Runtime.Type'. */
 	UPROPERTY(Config, EditAnywhere, Category = "NNE Denoiser", meta = (DisplayName = "Runtime Type", ToolTip = "Select a Runtime type", ConsoleVariable = "NNEDenoiser.Runtime.Type"))
diff --git a/Engine/Plugins/Experimental/NNEDenoiser/Source/NNEDenoiser/Public/NNEDenoiserTemporalAsset.h b/Engine/Plugins/Experimental/NNEDenoiser/Source/NNEDenoiser/Public/NNEDenoiserTemporalAsset.h
new file mode 100644
--- /dev/null
+++ b/Engine/Plugins/Experimental/NNEDenoiser/Source/NNEDenoiser/Public/NNEDenoiserTemporalAsset.h
@@ -0,0 +1,34 @@
+// Copyright Epic Games, Inc. All Rights Reserved.
+
+#pragma once
+
+#include "Engine/DataAsset.h"
+#include "Engine/DataTable.h"
+#include "NNEDenoiserTilingConfig.h"
+#include "NNEModelData.h"
+
+#include "NNEDenoiserTemporalAsset.generated.h"
+
+/** Denoiser model data asset */
+UCLASS(BlueprintType)
+class NNEDENOISER_API UNNEDenoiserTemporalAsset : public UDataAsset
+{
+	GENERATED_BODY()
+
+public:
+	/** NNE model data */
+	UPROPERTY(EditAnywhere, BlueprintReadWrite, Category=NNEDenoiser)
+	TSoftObjectPtr<UNNEModelData> ModelData;
+
+	/** Input mapping table */
+	UPROPERTY(EditAnywhere, BlueprintReadWrite, Category=NNEDenoiser, meta = (RequiredAssetDataTags = "RowStructure=/Script/NNEDenoiser.NNEDenoiserTemporalInputMappingData"))
+	TSoftObjectPtr<UDataTable> InputMapping;
+
+	/** Output mapping table */
+	UPROPERTY(EditAnywhere, BlueprintReadWrite, Category=NNEDenoiser, meta = (RequiredAssetDataTags = "RowStructure=/Script/NNEDenoiser.NNEDenoiserTemporalOutputMappingData"))
+	TSoftObjectPtr<UDataTable> OutputMapping;
+
+	/** Tiling configuration */
+	UPROPERTY(EditAnywhere, BlueprintReadWrite, Category=NNEDenoiser)
+	FTilingConfig TilingConfig{};
+};
\ No newline at end of file
diff --git a/Engine/Plugins/Experimental/NNEDenoiser/Source/NNEDenoiser/Public/NNEDenoiserTilingConfig.h b/Engine/Plugins/Experimental/NNEDenoiser/Source/NNEDenoiser/Public/NNEDenoiserTilingConfig.h
new file mode 100644
--- /dev/null
+++ b/Engine/Plugins/Experimental/NNEDenoiser/Source/NNEDenoiser/Public/NNEDenoiserTilingConfig.h
@@ -0,0 +1,28 @@
+// Copyright Epic Games, Inc. All Rights Reserved.
+
+#pragma once
+
+#include "NNEDenoiserTilingConfig.generated.h"
+
+/** Tiling configuration for fixed and dynamic size models */
+USTRUCT(BlueprintType)
+struct FTilingConfig
+{
+	GENERATED_BODY()
+
+	/** Tile size alignment (applies only to dynamic size models) */
+	UPROPERTY(EditAnywhere, BlueprintReadWrite, Category=NNEDenoiser, meta = (DisplayName = "Size Alignment"))
+	int32 Alignment = 1;
+
+	/** Tile overlap */
+	UPROPERTY(EditAnywhere, BlueprintReadWrite, Category=NNEDenoiser)
+	int32 Overlap = 0;
+
+	/** Maximum tile size (applies only to dynamic size models) */
+	UPROPERTY(EditAnywhere, BlueprintReadWrite, Category=NNEDenoiser)
+	int32 MaxSize = 0;
+
+	/** Minimum tile size (applies only to dynamic size models) */
+	UPROPERTY(EditAnywhere, BlueprintReadWrite, Category=NNEDenoiser)
+	int32 MinSize = 1;
+};
\ No newline at end of file
diff --git a/Engine/Source/Runtime/Renderer/Private/PathTracingSpatialTemporalDenoising.cpp b/Engine/Source/Runtime/Renderer/Private/PathTracingSpatialTemporalDenoising.cpp
--- a/Engine/Source/Runtime/Renderer/Private/PathTracingSpatialTemporalDenoising.cpp
+++ b/Engine/Source/Runtime/Renderer/Private/PathTracingSpatialTemporalDenoising.cpp
@@ -185,7 +185,7 @@ namespace {
 
 	TAutoConsoleVariable<FString> CVarPathTracingDenoiserName(
 		TEXT("r.PathTracing.Denoiser.Name"),
-		"NNE_OIDN",
+		"NNEDenoiser",
 		TEXT("Set the spatial denoiser name. It is the corresponding name registered by a denoiser plugin\n")
 		TEXT("Any registered denoiser should be able to denoise a single frame spatially.\n")
 	);
