[NNEDenoiser] Fix for compilation on Mac.

#rnx

diff --git a/Engine/Plugins/Experimental/NNEDenoiser/Source/NNEDenoiser/Private/NNEDenoiserResourceMapping.cpp b/Engine/Plugins/Experimental/NNEDenoiser/Source/NNEDenoiser/Private/NNEDenoiserResourceMapping.cpp
--- a/Engine/Plugins/Experimental/NNEDenoiser/Source/NNEDenoiser/Private/NNEDenoiserResourceMapping.cpp
+++ b/Engine/Plugins/Experimental/NNEDenoiser/Source/NNEDenoiser/Private/NNEDenoiserResourceMapping.cpp
@@ -105,7 +105,7 @@ TMap<int32, TArray<FChannelMapping>> FResourceMapping::GetChannelMappingPerFrame
 
 		if (Info.Name == Name)
 		{
-			Result.FindOrAdd(Info.Frame).Emplace(I, Info.Channel);
+			Result.FindOrAdd(Info.Frame).Add({I, Info.Channel});
 		}
 	}
 
