[NNEDenoiser] Set default denoiser asset to OIDN balanced with alpha.

#tests Editor, Win64 run default denoiser model

#rnx

diff --git a/Engine/Plugins/Experimental/NNEDenoiser/Source/NNEDenoiser/Private/NNEDenoiserSettings.cpp b/Engine/Plugins/Experimental/NNEDenoiser/Source/NNEDenoiser/Private/NNEDenoiserSettings.cpp
--- a/Engine/Plugins/Experimental/NNEDenoiser/Source/NNEDenoiser/Private/NNEDenoiserSettings.cpp
+++ b/Engine/Plugins/Experimental/NNEDenoiser/Source/NNEDenoiser/Private/NNEDenoiserSettings.cpp
@@ -11,8 +11,8 @@ UNNEDenoiserSettings::UNNEDenoiserSettings()
 	CategoryName = TEXT("Plugins");
 	SectionName = TEXT("NNE Denoiser");
 
-	DenoiserAsset = FSoftObjectPath(TEXT("/NNEDenoiser/NNED_Oidn2-3_Balanced.NNED_Oidn2-3_Balanced"));
-	TemporalDenoiserAsset = FSoftObjectPath(TEXT("/NNEDenoiser/NNEDT_Oidn2-3_Balanced.NNEDT_Oidn2-3_Balanced"));
+	DenoiserAsset = FSoftObjectPath(TEXT("/NNEDenoiser/NNED_Oidn2-3_Balanced_Alpha.NNED_Oidn2-3_Balanced_Alpha"));
+	TemporalDenoiserAsset = FSoftObjectPath(TEXT("/NNEDenoiser/NNEDT_Oidn2-3_Balanced_Alpha.NNEDT_Oidn2-3_Balanced_Alpha"));
 }
 
 void UNNEDenoiserSettings::PostInitProperties()
