[NNEDenoiser] Hide temporal asset selection from project settings.

#tests Editor, Win64 verify property does not show up in project settings

#rnx

diff --git a/Engine/Plugins/Experimental/NNEDenoiser/Source/NNEDenoiser/Public/NNEDenoiserSettings.h b/Engine/Plugins/Experimental/NNEDenoiser/Source/NNEDenoiser/Public/NNEDenoiserSettings.h
--- a/Engine/Plugins/Experimental/NNEDenoiser/Source/NNEDenoiser/Public/NNEDenoiserSettings.h
+++ b/Engine/Plugins/Experimental/NNEDenoiser/Source/NNEDenoiser/Public/NNEDenoiserSettings.h
@@ -32,8 +32,7 @@ public:
 	UPROPERTY(Config, EditAnywhere, Category = "NNE Denoiser", meta = (DisplayName = "Denoiser Asset", ToolTip = "Select the denoiser asset"))
 	TSoftObjectPtr<UNNEDenoiserAsset> DenoiserAsset;
 
-	/** Temporal denoiser asset data used to create a NNE Denoiser */
-	UPROPERTY(Config, EditAnywhere, Category = "NNE Denoiser", meta = (DisplayName = "Temporal Denoiser Asset", ToolTip = "Select the temporal denoiser asset"))
+	/** Temporal denoiser asset data used to create a NNE Denoiser (Currently not used and therefore "hidden") */
 	TSoftObjectPtr<UNNEDenoiserTemporalAsset> TemporalDenoiserAsset;
 
 private:
