[NNEDenoiser] Beta: move plugin out of Experimental folder and update .uplugin.

#tests Editor, Win64 run denoiser

#virtualized

diff --git a/Engine/Plugins/Experimental/NNEDenoiser/NNEDenoiser.uplugin b/Engine/Plugins/Experimental/NNEDenoiser/NNEDenoiser.uplugin
deleted file mode 100644
diff --git a/Engine/Plugins/Experimental/NNEDenoiser/Shaders/NNEDenoiserShadersAutoExposure.usf b/Engine/Plugins/Experimental/NNEDenoiser/Shaders/NNEDenoiserShadersAutoExposure.usf
deleted file mode 100644
diff --git a/Engine/Plugins/Experimental/NNEDenoiser/Shaders/NNEDenoiserShadersAutoExposureDownsample.usf b/Engine/Plugins/Experimental/NNEDenoiser/Shaders/NNEDenoiserShadersAutoExposureDownsample.usf
deleted file mode 100644
diff --git a/Engine/Plugins/Experimental/NNEDenoiser/Shaders/NNEDenoiserShadersAutoExposureReduce.usf b/Engine/Plugins/Experimental/NNEDenoiser/Shaders/NNEDenoiserShadersAutoExposureReduce.usf
deleted file mode 100644
diff --git a/Engine/Plugins/Experimental/NNEDenoiser/Shaders/NNEDenoiserShadersAutoExposureReduceFinal.usf b/Engine/Plugins/Experimental/NNEDenoiser/Shaders/NNEDenoiserShadersAutoExposureReduceFinal.usf
deleted file mode 100644
diff --git a/Engine/Plugins/Experimental/NNEDenoiser/Shaders/NNEDenoiserShadersDefaultIOProcess.usf b/Engine/Plugins/Experimental/NNEDenoiser/Shaders/NNEDenoiserShadersDefaultIOProcess.usf
deleted file mode 100644
diff --git a/Engine/Plugins/Experimental/NNEDenoiser/Shaders/NNEDenoiserShadersMappedCopy.usf b/Engine/Plugins/Experimental/NNEDenoiser/Shaders/NNEDenoiserShadersMappedCopy.usf
deleted file mode 100644
diff --git a/Engine/Plugins/Experimental/NNEDenoiser/Shaders/NNEDenoiserShadersTransferFunctionOidn.usf b/Engine/Plugins/Experimental/NNEDenoiser/Shaders/NNEDenoiserShadersTransferFunctionOidn.usf
deleted file mode 100644
diff --git a/Engine/Plugins/Experimental/NNEDenoiser/Source/NNEDenoiser/NNEDenoiser.Build.cs b/Engine/Plugins/Experimental/NNEDenoiser/Source/NNEDenoiser/NNEDenoiser.Build.cs
deleted file mode 100644
diff --git a/Engine/Plugins/Experimental/NNEDenoiser/Source/NNEDenoiser/Private/NNEDenoise.cpp b/Engine/Plugins/Experimental/NNEDenoiser/Source/NNEDenoiser/Private/NNEDenoise.cpp
deleted file mode 100644
diff --git a/Engine/Plugins/Experimental/NNEDenoiser/Source/NNEDenoiser/Private/NNEDenoiser.h b/Engine/Plugins/Experimental/NNEDenoiser/Source/NNEDenoiser/Private/NNEDenoiser.h
deleted file mode 100644
diff --git a/Engine/Plugins/Experimental/NNEDenoiser/Source/NNEDenoiser/Private/NNEDenoiserAutoExposure.cpp b/Engine/Plugins/Experimental/NNEDenoiser/Source/NNEDenoiser/Private/NNEDenoiserAutoExposure.cpp
deleted file mode 100644
diff --git a/Engine/Plugins/Experimental/NNEDenoiser/Source/NNEDenoiser/Private/NNEDenoiserAutoExposure.h b/Engine/Plugins/Experimental/NNEDenoiser/Source/NNEDenoiser/Private/NNEDenoiserAutoExposure.h
deleted file mode 100644
diff --git a/Engine/Plugins/Experimental/NNEDenoiser/Source/NNEDenoiser/Private/NNEDenoiserGenericDenoiser.cpp b/Engine/Plugins/Experimental/NNEDenoiser/Source/NNEDenoiser/Private/NNEDenoiserGenericDenoiser.cpp
deleted file mode 100644
diff --git a/Engine/Plugins/Experimental/NNEDenoiser/Source/NNEDenoiser/Private/NNEDenoiserGenericDenoiser.h b/Engine/Plugins/Experimental/NNEDenoiser/Source/NNEDenoiser/Private/NNEDenoiserGenericDenoiser.h
deleted file mode 100644
diff --git a/Engine/Plugins/Experimental/NNEDenoiser/Source/NNEDenoiser/Private/NNEDenoiserHistory.h b/Engine/Plugins/Experimental/NNEDenoiser/Source/NNEDenoiser/Private/NNEDenoiserHistory.h
deleted file mode 100644
diff --git a/Engine/Plugins/Experimental/NNEDenoiser/Source/NNEDenoiser/Private/NNEDenoiserIOMappingData.cpp b/Engine/Plugins/Experimental/NNEDenoiser/Source/NNEDenoiser/Private/NNEDenoiserIOMappingData.cpp
deleted file mode 100644
diff --git a/Engine/Plugins/Experimental/NNEDenoiser/Source/NNEDenoiser/Private/NNEDenoiserIOProcess.h b/Engine/Plugins/Experimental/NNEDenoiser/Source/NNEDenoiser/Private/NNEDenoiserIOProcess.h
deleted file mode 100644
diff --git a/Engine/Plugins/Experimental/NNEDenoiser/Source/NNEDenoiser/Private/NNEDenoiserIOProcessBase.cpp b/Engine/Plugins/Experimental/NNEDenoiser/Source/NNEDenoiser/Private/NNEDenoiserIOProcessBase.cpp
deleted file mode 100644
diff --git a/Engine/Plugins/Experimental/NNEDenoiser/Source/NNEDenoiser/Private/NNEDenoiserIOProcessBase.h b/Engine/Plugins/Experimental/NNEDenoiser/Source/NNEDenoiser/Private/NNEDenoiserIOProcessBase.h
deleted file mode 100644
diff --git a/Engine/Plugins/Experimental/NNEDenoiser/Source/NNEDenoiser/Private/NNEDenoiserLog.h b/Engine/Plugins/Experimental/NNEDenoiser/Source/NNEDenoiser/Private/NNEDenoiserLog.h
deleted file mode 100644
diff --git a/Engine/Plugins/Experimental/NNEDenoiser/Source/NNEDenoiser/Private/NNEDenoiserModelInstance.h b/Engine/Plugins/Experimental/NNEDenoiser/Source/NNEDenoiser/Private/NNEDenoiserModelInstance.h
deleted file mode 100644
diff --git a/Engine/Plugins/Experimental/NNEDenoiser/Source/NNEDenoiser/Private/NNEDenoiserModelInstanceCPU.cpp b/Engine/Plugins/Experimental/NNEDenoiser/Source/NNEDenoiser/Private/NNEDenoiserModelInstanceCPU.cpp
deleted file mode 100644
diff --git a/Engine/Plugins/Experimental/NNEDenoiser/Source/NNEDenoiser/Private/NNEDenoiserModelInstanceCPU.h b/Engine/Plugins/Experimental/NNEDenoiser/Source/NNEDenoiser/Private/NNEDenoiserModelInstanceCPU.h
deleted file mode 100644
diff --git a/Engine/Plugins/Experimental/NNEDenoiser/Source/NNEDenoiser/Private/NNEDenoiserModelInstanceGPU.cpp b/Engine/Plugins/Experimental/NNEDenoiser/Source/NNEDenoiser/Private/NNEDenoiserModelInstanceGPU.cpp
deleted file mode 100644
diff --git a/Engine/Plugins/Experimental/NNEDenoiser/Source/NNEDenoiser/Private/NNEDenoiserModelInstanceGPU.h b/Engine/Plugins/Experimental/NNEDenoiser/Source/NNEDenoiser/Private/NNEDenoiserModelInstanceGPU.h
deleted file mode 100644
diff --git a/Engine/Plugins/Experimental/NNEDenoiser/Source/NNEDenoiser/Private/NNEDenoiserModelInstanceRDG.cpp b/Engine/Plugins/Experimental/NNEDenoiser/Source/NNEDenoiser/Private/NNEDenoiserModelInstanceRDG.cpp
deleted file mode 100644
diff --git a/Engine/Plugins/Experimental/NNEDenoiser/Source/NNEDenoiser/Private/NNEDenoiserModelInstanceRDG.h b/Engine/Plugins/Experimental/NNEDenoiser/Source/NNEDenoiser/Private/NNEDenoiserModelInstanceRDG.h
deleted file mode 100644
diff --git a/Engine/Plugins/Experimental/NNEDenoiser/Source/NNEDenoiser/Private/NNEDenoiserParameters.h b/Engine/Plugins/Experimental/NNEDenoiser/Source/NNEDenoiser/Private/NNEDenoiserParameters.h
deleted file mode 100644
diff --git a/Engine/Plugins/Experimental/NNEDenoiser/Source/NNEDenoiser/Private/NNEDenoiserPathTracingDenoiser.cpp b/Engine/Plugins/Experimental/NNEDenoiser/Source/NNEDenoiser/Private/NNEDenoiserPathTracingDenoiser.cpp
deleted file mode 100644
diff --git a/Engine/Plugins/Experimental/NNEDenoiser/Source/NNEDenoiser/Private/NNEDenoiserPathTracingDenoiser.h b/Engine/Plugins/Experimental/NNEDenoiser/Source/NNEDenoiser/Private/NNEDenoiserPathTracingDenoiser.h
deleted file mode 100644
diff --git a/Engine/Plugins/Experimental/NNEDenoiser/Source/NNEDenoiser/Private/NNEDenoiserPathTracingSpatialTemporalDenoiser.cpp b/Engine/Plugins/Experimental/NNEDenoiser/Source/NNEDenoiser/Private/NNEDenoiserPathTracingSpatialTemporalDenoiser.cpp
deleted file mode 100644
diff --git a/Engine/Plugins/Experimental/NNEDenoiser/Source/NNEDenoiser/Private/NNEDenoiserPathTracingSpatialTemporalDenoiser.h b/Engine/Plugins/Experimental/NNEDenoiser/Source/NNEDenoiser/Private/NNEDenoiserPathTracingSpatialTemporalDenoiser.h
deleted file mode 100644
diff --git a/Engine/Plugins/Experimental/NNEDenoiser/Source/NNEDenoiser/Private/NNEDenoiserResourceManager.cpp b/Engine/Plugins/Experimental/NNEDenoiser/Source/NNEDenoiser/Private/NNEDenoiserResourceManager.cpp
deleted file mode 100644
diff --git a/Engine/Plugins/Experimental/NNEDenoiser/Source/NNEDenoiser/Private/NNEDenoiserResourceManager.h b/Engine/Plugins/Experimental/NNEDenoiser/Source/NNEDenoiser/Private/NNEDenoiserResourceManager.h
deleted file mode 100644
diff --git a/Engine/Plugins/Experimental/NNEDenoiser/Source/NNEDenoiser/Private/NNEDenoiserResourceMapping.cpp b/Engine/Plugins/Experimental/NNEDenoiser/Source/NNEDenoiser/Private/NNEDenoiserResourceMapping.cpp
deleted file mode 100644
diff --git a/Engine/Plugins/Experimental/NNEDenoiser/Source/NNEDenoiser/Private/NNEDenoiserResourceMapping.h b/Engine/Plugins/Experimental/NNEDenoiser/Source/NNEDenoiser/Private/NNEDenoiserResourceMapping.h
deleted file mode 100644
diff --git a/Engine/Plugins/Experimental/NNEDenoiser/Source/NNEDenoiser/Private/NNEDenoiserSettings.cpp b/Engine/Plugins/Experimental/NNEDenoiser/Source/NNEDenoiser/Private/NNEDenoiserSettings.cpp
deleted file mode 100644
diff --git a/Engine/Plugins/Experimental/NNEDenoiser/Source/NNEDenoiser/Private/NNEDenoiserTiling.h b/Engine/Plugins/Experimental/NNEDenoiser/Source/NNEDenoiser/Private/NNEDenoiserTiling.h
deleted file mode 100644
diff --git a/Engine/Plugins/Experimental/NNEDenoiser/Source/NNEDenoiser/Private/NNEDenoiserTransferFunction.h b/Engine/Plugins/Experimental/NNEDenoiser/Source/NNEDenoiser/Private/NNEDenoiserTransferFunction.h
deleted file mode 100644
diff --git a/Engine/Plugins/Experimental/NNEDenoiser/Source/NNEDenoiser/Private/NNEDenoiserTransferFunctionOidn.cpp b/Engine/Plugins/Experimental/NNEDenoiser/Source/NNEDenoiser/Private/NNEDenoiserTransferFunctionOidn.cpp
deleted file mode 100644
diff --git a/Engine/Plugins/Experimental/NNEDenoiser/Source/NNEDenoiser/Private/NNEDenoiserTransferFunctionOidn.h b/Engine/Plugins/Experimental/NNEDenoiser/Source/NNEDenoiser/Private/NNEDenoiserTransferFunctionOidn.h
deleted file mode 100644
diff --git a/Engine/Plugins/Experimental/NNEDenoiser/Source/NNEDenoiser/Private/NNEDenoiserUnitTestOidn.cpp b/Engine/Plugins/Experimental/NNEDenoiser/Source/NNEDenoiser/Private/NNEDenoiserUnitTestOidn.cpp
deleted file mode 100644
diff --git a/Engine/Plugins/Experimental/NNEDenoiser/Source/NNEDenoiser/Private/NNEDenoiserUnitTests.cpp b/Engine/Plugins/Experimental/NNEDenoiser/Source/NNEDenoiser/Private/NNEDenoiserUnitTests.cpp
deleted file mode 100644
diff --git a/Engine/Plugins/Experimental/NNEDenoiser/Source/NNEDenoiser/Private/NNEDenoiserUtils.h b/Engine/Plugins/Experimental/NNEDenoiser/Source/NNEDenoiser/Private/NNEDenoiserUtils.h
deleted file mode 100644
diff --git a/Engine/Plugins/Experimental/NNEDenoiser/Source/NNEDenoiser/Private/NNEDenoiserViewExtension.cpp b/Engine/Plugins/Experimental/NNEDenoiser/Source/NNEDenoiser/Private/NNEDenoiserViewExtension.cpp
deleted file mode 100644
diff --git a/Engine/Plugins/Experimental/NNEDenoiser/Source/NNEDenoiser/Private/NNEDenoiserViewExtension.h b/Engine/Plugins/Experimental/NNEDenoiser/Source/NNEDenoiser/Private/NNEDenoiserViewExtension.h
deleted file mode 100644
diff --git a/Engine/Plugins/Experimental/NNEDenoiser/Source/NNEDenoiser/Public/NNEDenoiserAsset.h b/Engine/Plugins/Experimental/NNEDenoiser/Source/NNEDenoiser/Public/NNEDenoiserAsset.h
deleted file mode 100644
diff --git a/Engine/Plugins/Experimental/NNEDenoiser/Source/NNEDenoiser/Public/NNEDenoiserIOMappingData.h b/Engine/Plugins/Experimental/NNEDenoiser/Source/NNEDenoiser/Public/NNEDenoiserIOMappingData.h
deleted file mode 100644
diff --git a/Engine/Plugins/Experimental/NNEDenoiser/Source/NNEDenoiser/Public/NNEDenoiserResourceName.h b/Engine/Plugins/Experimental/NNEDenoiser/Source/NNEDenoiser/Public/NNEDenoiserResourceName.h
deleted file mode 100644
diff --git a/Engine/Plugins/Experimental/NNEDenoiser/Source/NNEDenoiser/Public/NNEDenoiserSettings.h b/Engine/Plugins/Experimental/NNEDenoiser/Source/NNEDenoiser/Public/NNEDenoiserSettings.h
deleted file mode 100644
diff --git a/Engine/Plugins/Experimental/NNEDenoiser/Source/NNEDenoiser/Public/NNEDenoiserTemporalAsset.h b/Engine/Plugins/Experimental/NNEDenoiser/Source/NNEDenoiser/Public/NNEDenoiserTemporalAsset.h
deleted file mode 100644
diff --git a/Engine/Plugins/Experimental/NNEDenoiser/Source/NNEDenoiser/Public/NNEDenoiserTilingConfig.h b/Engine/Plugins/Experimental/NNEDenoiser/Source/NNEDenoiser/Public/NNEDenoiserTilingConfig.h
deleted file mode 100644
diff --git a/Engine/Plugins/Experimental/NNEDenoiser/Source/NNEDenoiserShaders/Internal/NNEDenoiserShadersAutoExposureCS.h b/Engine/Plugins/Experimental/NNEDenoiser/Source/NNEDenoiserShaders/Internal/NNEDenoiserShadersAutoExposureCS.h
deleted file mode 100644
diff --git a/Engine/Plugins/Experimental/NNEDenoiser/Source/NNEDenoiserShaders/Internal/NNEDenoiserShadersDefaultIOProcessCS.h b/Engine/Plugins/Experimental/NNEDenoiser/Source/NNEDenoiserShaders/Internal/NNEDenoiserShadersDefaultIOProcessCS.h
deleted file mode 100644
diff --git a/Engine/Plugins/Experimental/NNEDenoiser/Source/NNEDenoiserShaders/Internal/NNEDenoiserShadersMappedCopyCS.h b/Engine/Plugins/Experimental/NNEDenoiser/Source/NNEDenoiserShaders/Internal/NNEDenoiserShadersMappedCopyCS.h
deleted file mode 100644
diff --git a/Engine/Plugins/Experimental/NNEDenoiser/Source/NNEDenoiserShaders/Internal/NNEDenoiserShadersTransferFunctionOidnCS.h b/Engine/Plugins/Experimental/NNEDenoiser/Source/NNEDenoiserShaders/Internal/NNEDenoiserShadersTransferFunctionOidnCS.h
deleted file mode 100644
diff --git a/Engine/Plugins/Experimental/NNEDenoiser/Source/NNEDenoiserShaders/NNEDenoiserShaders.Build.cs b/Engine/Plugins/Experimental/NNEDenoiser/Source/NNEDenoiserShaders/NNEDenoiserShaders.Build.cs
deleted file mode 100644
diff --git a/Engine/Plugins/Experimental/NNEDenoiser/Source/NNEDenoiserShaders/Private/NNEDenoiserShadersAutoExposureCS.cpp b/Engine/Plugins/Experimental/NNEDenoiser/Source/NNEDenoiserShaders/Private/NNEDenoiserShadersAutoExposureCS.cpp
deleted file mode 100644
diff --git a/Engine/Plugins/Experimental/NNEDenoiser/Source/NNEDenoiserShaders/Private/NNEDenoiserShadersDefaultIOProcessCS.cpp b/Engine/Plugins/Experimental/NNEDenoiser/Source/NNEDenoiserShaders/Private/NNEDenoiserShadersDefaultIOProcessCS.cpp
deleted file mode 100644
diff --git a/Engine/Plugins/Experimental/NNEDenoiser/Source/NNEDenoiserShaders/Private/NNEDenoiserShadersLog.h b/Engine/Plugins/Experimental/NNEDenoiser/Source/NNEDenoiserShaders/Private/NNEDenoiserShadersLog.h
deleted file mode 100644
diff --git a/Engine/Plugins/Experimental/NNEDenoiser/Source/NNEDenoiserShaders/Private/NNEDenoiserShadersMappedCopyCS.cpp b/Engine/Plugins/Experimental/NNEDenoiser/Source/NNEDenoiserShaders/Private/NNEDenoiserShadersMappedCopyCS.cpp
deleted file mode 100644
diff --git a/Engine/Plugins/Experimental/NNEDenoiser/Source/NNEDenoiserShaders/Private/NNEDenoiserShadersModule.cpp b/Engine/Plugins/Experimental/NNEDenoiser/Source/NNEDenoiserShaders/Private/NNEDenoiserShadersModule.cpp
deleted file mode 100644
diff --git a/Engine/Plugins/Experimental/NNEDenoiser/Source/NNEDenoiserShaders/Private/NNEDenoiserShadersTransferFunctionOidnCS.cpp b/Engine/Plugins/Experimental/NNEDenoiser/Source/NNEDenoiserShaders/Private/NNEDenoiserShadersTransferFunctionOidnCS.cpp
deleted file mode 100644
diff --git a/Engine/Plugins/NNE/NNEDenoiser/NNEDenoiser.uplugin b/Engine/Plugins/NNE/NNEDenoiser/NNEDenoiser.uplugin
new file mode 100644
--- /dev/null
+++ b/Engine/Plugins/NNE/NNEDenoiser/NNEDenoiser.uplugin
@@ -0,0 +1,39 @@
+{
+	"FileVersion": 3,
+	"Version": 1,
+	"VersionName": "1.0",
+	"FriendlyName": "NNEDenoiser",
+	"Description": "Neural denoiser for the Unreal Path Tracer based on the Neural Network Engine (NNE).",
+	"Category": "Denoising",
+	"CreatedBy": "Epic Games, Inc.",
+	"CreatedByURL": "https://epicgames.com",
+	"DocsURL": "",
+	"MarketplaceURL": "",
+	"SupportURL": "",
+	"EnabledByDefault" : true,
+	"CanContainContent": true,
+	"IsBetaVersion": true,
+	"IsExperimentalVersion": false,
+	"Installed": false,
+	"SupportedTargetPlatforms": [
+		"Win64", "Linux", "Mac"
+	],
+	"Modules": [
+		{
+			"Name": "NNEDenoiser",
+			"Type": "Runtime",
+			"LoadingPhase": "Default"
+		},
+		{
+			"Name": "NNEDenoiserShaders",
+			"Type": "Runtime",
+			"LoadingPhase": "PostConfigInit"
+		}
+	],
+	"Plugins": [
+		{
+			"Name": "NNERuntimeORT",
+			"Enabled": true
+		}
+	]
+}
\ No newline at end of file
diff --git a/Engine/Plugins/NNE/NNEDenoiser/Shaders/NNEDenoiserShadersAutoExposure.usf b/Engine/Plugins/NNE/NNEDenoiser/Shaders/NNEDenoiserShadersAutoExposure.usf
new file mode 100644
--- /dev/null
+++ b/Engine/Plugins/NNE/NNEDenoiser/Shaders/NNEDenoiserShadersAutoExposure.usf
@@ -0,0 +1,40 @@
+// Copyright Epic Games, Inc. All Rights Reserved.
+
+#include "/Engine/Public/Platform.ush"
+#include "/Engine/Private/Common.ush"
+
+int InputTextureWidth;
+int InputTextureHeight;
+Texture2D InputTexture;
+RWBuffer<float> OutputBuffer;
+
+[numthreads(1, 1, 1)]
+void AutoExposure(in const uint3 DispatchThreadID : SV_DispatchThreadID)
+{
+	const float Key = 0.18;
+
+	float Sum = 0.0;
+	int Count = 0;
+	for (int i = 0; i < InputTextureWidth; i++)
+	{
+		for (int j = 0; j < InputTextureHeight; j++)
+		{
+			float3 Color = InputTexture.Load(int3(i, j, 0)).rgb;
+
+			Color = clamp(Color, float3(0.0, 0.0, 0.0), float3(MAX_FLT, MAX_FLT, MAX_FLT));
+
+			float Lum = Luminance(Color);
+			if (Lum > 1e-8)
+			{
+				Sum += log2(Lum);
+				Count++;
+			}
+		}
+	}
+
+	float Result = Count > 0 ? Sum / (float)Count : 1.0;
+	Result = Key / exp2(Result);
+
+	OutputBuffer[0] = Result;
+	OutputBuffer[1] = 1.0 / Result;
+}
\ No newline at end of file
diff --git a/Engine/Plugins/NNE/NNEDenoiser/Shaders/NNEDenoiserShadersAutoExposureDownsample.usf b/Engine/Plugins/NNE/NNEDenoiser/Shaders/NNEDenoiserShadersAutoExposureDownsample.usf
new file mode 100644
--- /dev/null
+++ b/Engine/Plugins/NNE/NNEDenoiser/Shaders/NNEDenoiserShadersAutoExposureDownsample.usf
@@ -0,0 +1,55 @@
+// Copyright Epic Games, Inc. All Rights Reserved.
+
+#include "/Engine/Public/Platform.ush"
+#include "/Engine/Private/Common.ush"
+
+int InputTextureWidth;
+int InputTextureHeight;
+Texture2D InputTexture;
+
+int NumBinsW;
+int NumBinsH;
+RWBuffer<float> OutputBins;
+
+groupshared float Sums[THREAD_GROUP_SIZE * THREAD_GROUP_SIZE];
+
+[numthreads(THREAD_GROUP_SIZE, THREAD_GROUP_SIZE, 1)]
+void Downsample(in const uint3 GroupThreadID : SV_GroupThreadID, in const uint3 GroupID : SV_GroupID, in const uint3 DispatchThreadID : SV_DispatchThreadID)
+{
+	const int BeginW = GroupID.x * InputTextureWidth / NumBinsW;
+	const int BeginH = GroupID.y * InputTextureHeight / NumBinsH;
+	const int EndW = (GroupID.x + 1) * InputTextureWidth / NumBinsW;
+	const int EndH = (GroupID.y + 1) * InputTextureHeight / NumBinsH;
+
+	const int CoordW = BeginW + GroupThreadID.x;
+	const int CoordH = BeginH + GroupThreadID.y;
+
+	float Lum = 0.0;
+	if (CoordW < EndW && CoordH < EndH)
+	{
+		float3 Color = InputTexture.Load(int3(CoordW, CoordH, 0)).rgb;
+		Color = clamp(Color, float3(0.0, 0.0, 0.0), float3(MAX_FLT, MAX_FLT, MAX_FLT));
+
+		Lum = Luminance(Color);
+	}
+
+	const int LocalID = GroupThreadID.y * THREAD_GROUP_SIZE + GroupThreadID.x;
+	Sums[LocalID] = Lum;
+
+	for (int i = THREAD_GROUP_SIZE * THREAD_GROUP_SIZE / 2; i > 0; i >>= 1)
+	{
+		GroupMemoryBarrierWithGroupSync();
+		
+		if (LocalID < i)
+		{
+			Sums[LocalID] += Sums[LocalID + i];
+		}
+	}
+
+	if (LocalID == 0)
+	{
+		const float AverageLum = Sums[0] / float((EndW - BeginW) * (EndH - BeginH));
+		const int GroupLinearID = GroupID.y * NumBinsW + GroupID.x;
+		OutputBins[GroupLinearID] = AverageLum;
+	}
+}
\ No newline at end of file
diff --git a/Engine/Plugins/NNE/NNEDenoiser/Shaders/NNEDenoiserShadersAutoExposureReduce.usf b/Engine/Plugins/NNE/NNEDenoiser/Shaders/NNEDenoiserShadersAutoExposureReduce.usf
new file mode 100644
--- /dev/null
+++ b/Engine/Plugins/NNE/NNEDenoiser/Shaders/NNEDenoiserShadersAutoExposureReduce.usf
@@ -0,0 +1,50 @@
+// Copyright Epic Games, Inc. All Rights Reserved.
+
+#include "/Engine/Public/Platform.ush"
+#include "/Engine/Private/Common.ush"
+
+int InputSize;
+Buffer<float> InputBins;
+RWBuffer<float> OutputSums;
+RWBuffer<int> OutputCounts;
+int NumThreads;
+
+groupshared float Sums[THREAD_GROUP_SIZE];
+groupshared int Counts[THREAD_GROUP_SIZE];
+
+[numthreads(THREAD_GROUP_SIZE, 1, 1)]
+void Reduce(in const uint3 GroupThreadID : SV_GroupThreadID, in const uint3 GroupID : SV_GroupID, in const uint3 DispatchThreadID : SV_DispatchThreadID)
+{
+	float LocalSum = 0.0;
+	int LocalCount = 0;
+	for (int i = DispatchThreadID.x; i < InputSize; i += NumThreads)
+	{
+		const float Lum = InputBins[i];
+		if (Lum > EPS)
+		{
+			LocalSum += log2(Lum);
+			++LocalCount;
+		}
+	}
+
+	const int LocalID = GroupThreadID.x;
+	Sums[LocalID] = LocalSum;
+	Counts[LocalID] = LocalCount;
+
+	for (int i = THREAD_GROUP_SIZE / 2; i > 0; i >>= 1)
+	{
+		GroupMemoryBarrierWithGroupSync();
+
+		if (LocalID < i)
+		{
+			Sums[LocalID] += Sums[LocalID + i];
+			Counts[LocalID] += Counts[LocalID + i];
+		}
+	}
+
+	if (LocalID == 0)
+	{
+		OutputSums[GroupID.x] = Sums[0];
+		OutputCounts[GroupID.x] = Counts[0];
+	}
+}
\ No newline at end of file
diff --git a/Engine/Plugins/NNE/NNEDenoiser/Shaders/NNEDenoiserShadersAutoExposureReduceFinal.usf b/Engine/Plugins/NNE/NNEDenoiser/Shaders/NNEDenoiserShadersAutoExposureReduceFinal.usf
new file mode 100644
--- /dev/null
+++ b/Engine/Plugins/NNE/NNEDenoiser/Shaders/NNEDenoiserShadersAutoExposureReduceFinal.usf
@@ -0,0 +1,46 @@
+// Copyright Epic Games, Inc. All Rights Reserved.
+
+#include "/Engine/Public/Platform.ush"
+#include "/Engine/Private/Common.ush"
+
+int InputSize;
+Buffer<float> InputSums;
+Buffer<int> InputCounts;
+RWBuffer<float> OutputBuffer; // 0: Result, 1: 1.0 / Result
+
+groupshared float Sums[THREAD_GROUP_SIZE];
+groupshared int Counts[THREAD_GROUP_SIZE];
+
+[numthreads(THREAD_GROUP_SIZE, 1, 1)]
+void ReduceFinal(in const uint3 GroupThreadID : SV_GroupThreadID)
+{
+	const int LocalID = GroupThreadID.x;
+
+	if (LocalID < InputSize)
+	{
+		Sums[LocalID] = InputSums[LocalID];
+		Counts[LocalID] = InputCounts[LocalID];
+	}
+	else
+	{
+		Sums[LocalID] = 0;
+		Counts[LocalID] = 0;
+	}
+
+	for (int i = THREAD_GROUP_SIZE / 2; i > 0; i >>= 1)
+	{
+		GroupMemoryBarrierWithGroupSync();
+
+		if (LocalID < i)
+		{
+			Sums[LocalID] += Sums[LocalID + i];
+			Counts[LocalID] += Counts[LocalID + i];
+		}
+	}
+
+	if (LocalID == 0)
+	{
+		OutputBuffer[0] = (Counts[0] > 0) ? (KEY / exp2(Sums[0] / float(Counts[0]))) : 1.0;
+		OutputBuffer[1] = 1.0 / OutputBuffer[0];
+	}
+}
\ No newline at end of file
diff --git a/Engine/Plugins/NNE/NNEDenoiser/Shaders/NNEDenoiserShadersDefaultIOProcess.usf b/Engine/Plugins/NNE/NNEDenoiser/Shaders/NNEDenoiserShadersDefaultIOProcess.usf
new file mode 100644
--- /dev/null
+++ b/Engine/Plugins/NNE/NNEDenoiser/Shaders/NNEDenoiserShadersDefaultIOProcess.usf
@@ -0,0 +1,40 @@
+// Copyright Epic Games, Inc. All Rights Reserved.
+
+#include "/Engine/Public/Platform.ush"
+#include "/Engine/Private/Common.ush"
+
+int Width;
+int Height;
+Texture2D InputTexture;
+RWTexture2D<float4> OutputTexture;
+
+float3 GetProcessedValue(float3 Value)
+{
+#if INPUT_KIND_INDEX == 0 // Input color
+	return clamp(Value, float3(0.0, 0.0, 0.0), float3(MAX_FLT, MAX_FLT, MAX_FLT));
+#elif INPUT_KIND_INDEX == 1 // Input albedo
+	return clamp(Value, float3(0.0, 0.0, 0.0), float3(1.0, 1.0, 1.0));
+#elif INPUT_KIND_INDEX == 2 // Input normal
+	return clamp(Value, float3(-1.0, -1.0, -1.0), float3(1.0, 1.0, 1.0)) * 0.5 + float3(0.5, 0.5, 0.5);
+#elif INPUT_KIND_INDEX == 3 // Input flow
+	// TODO flow texture check for any necessary conversion
+	return Value;
+#else
+	// INPUT_KIND_INDEX == 4 // Output color
+	return clamp(Value, float3(0.0, 0.0, 0.0), float3(MAX_FLT, MAX_FLT, MAX_FLT));
+#endif
+}
+
+
+[numthreads(THREAD_GROUP_SIZE, THREAD_GROUP_SIZE, 1)]
+void IOProcess(in const uint3 DispatchThreadID : SV_DispatchThreadID)
+{
+	if (DispatchThreadID.x >= Width || DispatchThreadID.y >= Height)
+	{
+		return;
+	}
+
+	const float4 Color = InputTexture.Load(int3(DispatchThreadID.x, DispatchThreadID.y, 0));
+
+	OutputTexture[DispatchThreadID.xy] = float4(GetProcessedValue(Color.rgb), Color.a);
+}
\ No newline at end of file
diff --git a/Engine/Plugins/NNE/NNEDenoiser/Shaders/NNEDenoiserShadersMappedCopy.usf b/Engine/Plugins/NNE/NNEDenoiser/Shaders/NNEDenoiserShadersMappedCopy.usf
new file mode 100644
--- /dev/null
+++ b/Engine/Plugins/NNE/NNEDenoiser/Shaders/NNEDenoiserShadersMappedCopy.usf
@@ -0,0 +1,102 @@
+// Copyright Epic Games, Inc. All Rights Reserved.
+
+#include "/Engine/Public/Platform.ush"
+#include "/Engine/Private/Common.ush"
+
+#define FLOAT16 half
+#define FLOAT16_4 half4
+
+#define FLOAT32 float
+#define FLOAT32_4 float4
+
+#if INPUT_DATA_TYPE_INDEX == 4 // Float16
+#define INPUT_SCALAR_DATA_TYPE FLOAT16
+#define INPUT_VECTOR_DATA_TYPE FLOAT16_4
+#else // INPUT_DATA_TYPE_INDEX == 3 // Float32
+#define INPUT_SCALAR_DATA_TYPE FLOAT32
+#define INPUT_VECTOR_DATA_TYPE FLOAT32_4
+#endif
+
+#if OUTPUT_DATA_TYPE_INDEX == 4 // Float16
+#define OUTPUT_SCALAR_DATA_TYPE FLOAT16
+#define OUTPUT_VECTOR_DATA_TYPE FLOAT16_4
+#else // OUTPUT_DATA_TYPE_INDEX == 3 // Float32
+#define OUTPUT_SCALAR_DATA_TYPE FLOAT32
+#define OUTPUT_VECTOR_DATA_TYPE FLOAT32_4
+#endif
+
+#if INPUT_DATA_TYPE_INDEX == 4 // Float16
+#if OUTPUT_DATA_TYPE_INDEX == 4 // Float16
+#define CONVERT(X) X
+#else // OUTPUT_DATA_TYPE_INDEX == 3 // Float32
+#define CONVERT(X) FLOAT32(X)
+#endif
+#else // INPUT_DATA_TYPE_INDEX == 3 // Float32
+#if OUTPUT_DATA_TYPE_INDEX == 4 // Float16
+#define CONVERT(X) FLOAT16(X)
+#else // OUTPUT_DATA_TYPE_INDEX == 3 // Float32
+#define CONVERT(X) X
+#endif
+#endif
+
+int Width;
+int Height;
+int4 OutputChannel_InputChannel_Unused_Unused[MAX_NUM_MAPPED_CHANNELS];
+
+#if INTRINSIC_INPUT_TYPE == 0 // Texture
+Texture2D<INPUT_VECTOR_DATA_TYPE> InputTexture;
+#else // INTRINSIC_INPUT_TYPE == 1 // Buffer
+Buffer<INPUT_SCALAR_DATA_TYPE> InputBuffer;
+#endif
+
+#if INTRINSIC_OUTPUT_TYPE == 0 // Texture
+RWTexture2D<OUTPUT_VECTOR_DATA_TYPE> OutputTexture;
+#else // INTRINSIC_OUTPUT_TYPE == 1 // Buffer
+RWBuffer<OUTPUT_SCALAR_DATA_TYPE> OutputBuffer;
+#endif
+
+[numthreads(THREAD_GROUP_SIZE, THREAD_GROUP_SIZE, 1)]
+void MappedCopy(in const uint3 DispatchThreadID : SV_DispatchThreadID)
+{
+	if (DispatchThreadID.x >= Width || DispatchThreadID.y >= Height)
+	{
+		return;
+	}
+
+#if INTRINSIC_INPUT_TYPE == 1 || INTRINSIC_OUTPUT_TYPE == 1
+	const int Idx = DispatchThreadID.y * Width + DispatchThreadID.x;
+	const int Offset = Width * Height;
+#endif
+
+#if INTRINSIC_INPUT_TYPE == 0
+	const INPUT_VECTOR_DATA_TYPE InputColor = InputTexture.Load(int3(DispatchThreadID.x, DispatchThreadID.y, 0));
+#endif
+
+#if INTRINSIC_OUTPUT_TYPE == 0
+	OUTPUT_VECTOR_DATA_TYPE Result = OutputTexture[DispatchThreadID.xy];
+#endif
+
+	UNROLL
+	for (int I = 0; I < NUM_MAPPED_CHANNELS; I++)
+	{
+		const int4 Info = OutputChannel_InputChannel_Unused_Unused[I];
+		const int OutputChannel = Info.x;
+		const int InputChannel = Info.y;
+
+#if INTRINSIC_INPUT_TYPE == 0
+		const INPUT_SCALAR_DATA_TYPE InputValue = InputColor[InputChannel];
+#else
+		const INPUT_SCALAR_DATA_TYPE InputValue = InputBuffer[Idx + InputChannel * Offset];
+#endif
+
+#if INTRINSIC_OUTPUT_TYPE == 0
+		Result[OutputChannel] = CONVERT(InputValue);
+#else
+		OutputBuffer[Idx + OutputChannel * Offset] = CONVERT(InputValue);
+#endif
+	}
+
+#if INTRINSIC_OUTPUT_TYPE == 0
+	OutputTexture[DispatchThreadID.xy] = Result;
+#endif
+}
\ No newline at end of file
diff --git a/Engine/Plugins/NNE/NNEDenoiser/Shaders/NNEDenoiserShadersTransferFunctionOidn.usf b/Engine/Plugins/NNE/NNEDenoiser/Shaders/NNEDenoiserShadersTransferFunctionOidn.usf
new file mode 100644
--- /dev/null
+++ b/Engine/Plugins/NNE/NNEDenoiser/Shaders/NNEDenoiserShadersTransferFunctionOidn.usf
@@ -0,0 +1,76 @@
+// Copyright Epic Games, Inc. All Rights Reserved.
+
+#include "/Engine/Public/Platform.ush"
+#include "/Engine/Private/Common.ush"
+
+int Width;
+int Height;
+Texture2D InputTexture;
+Buffer<float> InputScaleBuffer;
+float NormScale;
+float InvNormScale;
+RWTexture2D<float4> OutputTexture;
+
+
+///
+
+// from [REDACTED_PROFILE_URL]/oidn/blob/master/core/color.h
+// struct FPU
+// {
+static float a  =  1.41283765e+03;
+static float b  =  1.64593172e+00;
+static float c  =  4.31384981e-01;
+static float d  = -2.94139609e-03;
+static float e  =  1.92653254e-01;
+static float f  =  6.26026094e-03;
+static float g  =  9.98620152e-01;
+static float y0 =  1.57945760e-06;
+static float y1 =  3.22087631e-02;
+static float x0 =  2.23151711e-03;
+static float x1 =  3.70974749e-01;
+
+static float Forward(float y)
+{
+	if (y <= y0) return a * y;
+	else if (y <= y1) return b * pow(y, c) + d;
+	return e * log(y + f) + g;
+}
+
+static float Inverse(float x)
+{
+	if (x <= x0) return x / a;
+	else if (x <= x1) return pow((x - d) / b, 1.0 / c);
+	return exp((x - g) / e) - f;
+}
+
+float3 GetProcessedValue(float3 Value)
+{
+#if MODE == 0 // Forward
+	Value = Value * InputScaleBuffer[0];
+	Value = clamp(Value, float3(0.0, 0.0, 0.0), float3(MAX_FLT, MAX_FLT, MAX_FLT));
+
+	return float3(Forward(Value.r), Forward(Value.g), Forward(Value.b)) * NormScale;
+#else
+	// MODE == 1 // Inverse
+	Value = clamp(Value, float3(0.0, 0.0, 0.0), float3(MAX_FLT, MAX_FLT, MAX_FLT)) * InvNormScale;
+	
+	return float3(Inverse(Value.r), Inverse(Value.g), Inverse(Value.b)) * InputScaleBuffer[1];
+#endif
+}
+
+// };
+
+///
+
+[numthreads(THREAD_GROUP_SIZE, THREAD_GROUP_SIZE, 1)]
+void PreOrPostprocess(in const uint3 DispatchThreadID : SV_DispatchThreadID)
+{
+	if (DispatchThreadID.x >= Width || DispatchThreadID.y >= Height)
+	{
+		return;
+	}
+
+	const float4 Color = InputTexture.Load(int3(DispatchThreadID.x, DispatchThreadID.y, 0));
+
+	OutputTexture[DispatchThreadID.xy] = float4(GetProcessedValue(Color.rgb), Color.a);
+}
\ No newline at end of file
diff --git a/Engine/Plugins/NNE/NNEDenoiser/Source/NNEDenoiser/NNEDenoiser.Build.cs b/Engine/Plugins/NNE/NNEDenoiser/Source/NNEDenoiser/NNEDenoiser.Build.cs
new file mode 100644
--- /dev/null
+++ b/Engine/Plugins/NNE/NNEDenoiser/Source/NNEDenoiser/NNEDenoiser.Build.cs
@@ -0,0 +1,30 @@
+// Copyright Epic Games, Inc. All Rights Reserved.
+
+using UnrealBuildTool;
+
+public class NNEDenoiser : ModuleRules
+{
+	public NNEDenoiser(ReadOnlyTargetRules Target) : base(Target)
+	{
+		PCHUsage = PCHUsageMode.UseExplicitOrSharedPCHs;
+
+		PublicDependencyModuleNames.AddRange(
+			new string[]
+			{
+				"Core",
+				"RenderCore"
+			});
+		
+		PrivateDependencyModuleNames.AddRange(
+			new string[]
+			{
+				"CoreUObject",
+				"DeveloperSettings",
+				"Engine",
+				"NNE",
+				"NNEDenoiserShaders",
+				"Renderer",
+				"RHI"
+			});
+	}
+}
diff --git a/Engine/Plugins/NNE/NNEDenoiser/Source/NNEDenoiser/Private/NNEDenoise.cpp b/Engine/Plugins/NNE/NNEDenoiser/Source/NNEDenoiser/Private/NNEDenoise.cpp
new file mode 100644
--- /dev/null
+++ b/Engine/Plugins/NNE/NNEDenoiser/Source/NNEDenoiser/Private/NNEDenoise.cpp
@@ -0,0 +1,37 @@
+// Copyright Epic Games, Inc. All Rights Reserved.
+
+#include "NNEDenoiser.h"
+#include "Misc/CoreDelegates.h"
+#include "NNEDenoiserLog.h"
+#include "NNEDenoiserViewExtension.h"
+
+DEFINE_LOG_CATEGORY(LogNNEDenoiser);
+
+void FNNEDenoiserModule::StartupModule()
+{
+	// During cook, we are not allowed to call IsRayTracingAllowed() later, therefore we assume Ray Tracing is not available and do not register the View Extension.
+	if (IsRunningCookCommandlet())
+	{
+		return;
+	}
+
+	FCoreDelegates::OnPostEngineInit.AddLambda([this] ()
+	{
+		// Only register View Extension if Ray Tracing is available.
+		if (IsRayTracingAllowed())
+		{
+			ViewExtension = FSceneViewExtensions::NewExtension<UE::NNEDenoiser::Private::FViewExtension>();
+		}
+		else
+		{
+			UE_LOG(LogNNEDenoiser, Log, TEXT("Ray Tracing is not enabled, therefore NNEDenoiser is not registered!"));
+		}
+	});
+}
+
+void FNNEDenoiserModule::ShutdownModule()
+{
+	ViewExtension.Reset();
+}
+
+IMPLEMENT_MODULE(FNNEDenoiserModule, NNEDenoiser)
\ No newline at end of file
diff --git a/Engine/Plugins/NNE/NNEDenoiser/Source/NNEDenoiser/Private/NNEDenoiser.h b/Engine/Plugins/NNE/NNEDenoiser/Source/NNEDenoiser/Private/NNEDenoiser.h
new file mode 100644
--- /dev/null
+++ b/Engine/Plugins/NNE/NNEDenoiser/Source/NNEDenoiser/Private/NNEDenoiser.h
@@ -0,0 +1,22 @@
+// Copyright Epic Games, Inc. All Rights Reserved.
+
+#pragma once
+
+#include "Modules/ModuleManager.h"
+#include "Templates/SharedPointer.h"
+
+namespace UE::NNEDenoiser::Private
+{
+	class FViewExtension;
+}
+
+class FNNEDenoiserModule : public IModuleInterface
+{
+public:
+	/** IModuleInterface implementation */
+	virtual void StartupModule() override;
+	virtual void ShutdownModule() override;
+
+private:
+	TSharedPtr<UE::NNEDenoiser::Private::FViewExtension, ESPMode::ThreadSafe> ViewExtension;
+};
diff --git a/Engine/Plugins/NNE/NNEDenoiser/Source/NNEDenoiser/Private/NNEDenoiserAutoExposure.cpp b/Engine/Plugins/NNE/NNEDenoiser/Source/NNEDenoiser/Private/NNEDenoiserAutoExposure.cpp
new file mode 100644
--- /dev/null
+++ b/Engine/Plugins/NNE/NNEDenoiser/Source/NNEDenoiser/Private/NNEDenoiserAutoExposure.cpp
@@ -0,0 +1,176 @@
+// Copyright Epic Games, Inc. All Rights Reserved.
+
+#include "NNEDenoiserAutoExposure.h"
+#include "NNEDenoiserLog.h"
+#include "NNEDenoiserShadersAutoExposureCS.h"
+#include "NNEDenoiserUtils.h"
+#include "RenderGraphBuilder.h"
+
+DECLARE_GPU_STAT_NAMED(FAutoExposure, TEXT("AutoExposure"));
+DECLARE_GPU_STAT_NAMED(FAutoExposureDownsample, TEXT("AutoExposure.Downsample"));
+DECLARE_GPU_STAT_NAMED(FAutoExposureReduce, TEXT("AutoExposure.Reduce"));
+DECLARE_GPU_STAT_NAMED(FAutoExposureReduceFinal, TEXT("AutoExposure.ReduceFinal"));
+
+BEGIN_SHADER_PARAMETER_STRUCT(FAutoExposureTestDownloadParameters, )
+	RDG_BUFFER_ACCESS(OutputBuffer, ERHIAccess::CopySrc)
+END_SHADER_PARAMETER_STRUCT()
+
+namespace UE::NNEDenoiser::Private
+{
+
+void FAutoExposure::Run(TConstArrayView<FLinearColor> InputData, FIntPoint Size, float& OutputValue) const
+{
+	using namespace UE::NNEDenoiserShaders::Internal;
+
+	OutputValue = 0.0f;
+	if (InputData.IsEmpty())
+	{
+		return;
+	}
+
+	const int32 NumBinsW = FMath::DivideAndRoundUp(Size.X, FAutoExposureDownsampleConstants::MAX_BIN_SIZE);
+	const int32 NumBinsH = FMath::DivideAndRoundUp(Size.Y, FAutoExposureDownsampleConstants::MAX_BIN_SIZE);
+	const int32 NumBins = NumBinsW * NumBinsH;
+
+	TArray<float> Bins;
+	Bins.SetNumUninitialized(NumBins);
+
+	for (int32 GroupX = 0; GroupX < NumBinsW; GroupX++)
+	{
+		for (int32 GroupY = 0; GroupY < NumBinsH; GroupY++)
+		{
+			const int32 BeginW = GroupX * Size.X / NumBinsW;
+			const int32 BeginH = GroupY * Size.Y / NumBinsH;
+			const int32 EndW = (GroupX + 1) * Size.X / NumBinsW;
+			const int32 EndH = (GroupY + 1) * Size.Y / NumBinsH;
+
+			float Lum = 0.0f;
+			for (int32 i = 0; i < FAutoExposureDownsampleConstants::MAX_BIN_SIZE; i++)
+			{
+				for (int32 j = 0; j < FAutoExposureDownsampleConstants::MAX_BIN_SIZE; j++)
+				{
+					const int32 CoordW = BeginW + i;
+					const int32 CoordH = BeginH + j;
+
+					if (CoordW < EndW && CoordH < EndH)
+					{
+						FLinearColor Color = InputData[CoordH * Size.X + CoordW];
+						
+						Lum += 0.212671f * Color.R + 0.715160f * Color.G + 0.072169f * Color.B;
+					}
+				}
+			}
+
+			Bins[GroupY * NumBinsW + GroupX] = Lum / float((EndW - BeginW) * (EndH - BeginH));
+		}
+	}
+
+	float Sum = 0.0;
+	int32 Count = 0;
+	for (int32 i = 0; i < NumBins; i++)
+	{
+		const float Lum = Bins[i];
+		if (Lum > FAutoExposureReduceConstants::EPS)
+		{
+			Sum += FMath::Log2(Lum);
+			Count++;
+		}
+	}
+
+	OutputValue = Count > 0 ? FAutoExposureReduceConstants::KEY / FMath::Exp2(Sum / (float)Count) : 1.0f;
+}
+
+void FAutoExposure::EnqueueRDG(FRDGBuilder& GraphBuilder, FRDGTextureRef InputTexture, FRDGBufferRef OutputBuffer) const
+{
+	using namespace UE::NNEDenoiserShaders::Internal;
+
+	const FIntVector InputTextureSize = InputTexture->Desc.GetSize();
+
+	const int32 NumBinsW = FMath::DivideAndRoundUp(InputTextureSize.X, FAutoExposureDownsampleConstants::MAX_BIN_SIZE);
+	const int32 NumBinsH = FMath::DivideAndRoundUp(InputTextureSize.Y, FAutoExposureDownsampleConstants::MAX_BIN_SIZE);
+	const int32 NumBins = NumBinsW * NumBinsH;
+
+	const int NumReduceGroups = FMath::Min(CeilDiv(NumBins, FAutoExposureReduceConstants::THREAD_GROUP_SIZE), FAutoExposureReduceConstants::THREAD_GROUP_SIZE);
+
+	FRDGBufferDesc OutputBinsDesc = FRDGBufferDesc::CreateBufferDesc(sizeof(float), NumBins);
+	FRDGBufferRef OutputBins = GraphBuilder.CreateBuffer(OutputBinsDesc, TEXT("AutoExposureOutputBins"));
+
+	FRDGBufferDesc OutputSumsDesc = FRDGBufferDesc::CreateBufferDesc(sizeof(float), NumReduceGroups);
+	FRDGBufferRef OutputSums = GraphBuilder.CreateBuffer(OutputSumsDesc, TEXT("AutoExposureOutputSums"));
+
+	FRDGBufferDesc OutputCountsDesc = FRDGBufferDesc::CreateBufferDesc(sizeof(int32), NumReduceGroups);
+	FRDGBufferRef OutputCounts = GraphBuilder.CreateBuffer(OutputCountsDesc, TEXT("AutoExposureOutputCounts"));
+
+	{
+		const FIntVector GroupCount = FIntVector(NumBinsW, NumBinsH, 1);
+
+		FAutoExposureDownsampleCS::FParameters *ShaderParameters = GraphBuilder.AllocParameters<FAutoExposureDownsampleCS::FParameters>();
+		ShaderParameters->InputTextureWidth = InputTextureSize.X;
+		ShaderParameters->InputTextureHeight = InputTextureSize.Y;
+		ShaderParameters->InputTexture = InputTexture;
+		ShaderParameters->NumBinsW = NumBinsW;
+		ShaderParameters->NumBinsH = NumBinsH;
+		ShaderParameters->OutputBins = GraphBuilder.CreateUAV(FRDGBufferUAVDesc(OutputBins, EPixelFormat::PF_R32_FLOAT));
+
+		FGlobalShaderMap* GlobalShaderMap = GetGlobalShaderMap(GMaxRHIFeatureLevel);
+		TShaderMapRef<FAutoExposureDownsampleCS> Shader(GlobalShaderMap);
+
+		RDG_EVENT_SCOPE(GraphBuilder, "AutoExposure.Downsample");
+		RDG_GPU_STAT_SCOPE(GraphBuilder, FAutoExposureDownsample);
+
+		FComputeShaderUtils::AddPass(
+			GraphBuilder,
+			RDG_EVENT_NAME("AutoExposure.Downsample"),
+			ERDGPassFlags::Compute | ERDGPassFlags::NeverCull,
+			Shader,
+			ShaderParameters,
+			GroupCount);
+	}
+
+	{
+		FAutoExposureReduceCS::FParameters *ShaderParameters = GraphBuilder.AllocParameters<FAutoExposureReduceCS::FParameters>();
+		ShaderParameters->InputSize = NumBins;
+		ShaderParameters->InputBins = GraphBuilder.CreateSRV(FRDGBufferSRVDesc(OutputBins, EPixelFormat::PF_R32_FLOAT));
+		ShaderParameters->OutputSums = GraphBuilder.CreateUAV(FRDGBufferUAVDesc(OutputSums, EPixelFormat::PF_R32_FLOAT));
+		ShaderParameters->OutputCounts = GraphBuilder.CreateUAV(FRDGBufferUAVDesc(OutputCounts, EPixelFormat::PF_R32_SINT));
+		ShaderParameters->NumThreads = NumReduceGroups * FAutoExposureDownsampleConstants::MAX_BIN_SIZE;
+
+		FGlobalShaderMap* GlobalShaderMap = GetGlobalShaderMap(GMaxRHIFeatureLevel);
+		TShaderMapRef<FAutoExposureReduceCS> Shader(GlobalShaderMap);
+
+		RDG_EVENT_SCOPE(GraphBuilder, "AutoExposure.Recude");
+		RDG_GPU_STAT_SCOPE(GraphBuilder, FAutoExposureReduce);
+
+		FComputeShaderUtils::AddPass(
+			GraphBuilder,
+			RDG_EVENT_NAME("AutoExposure.Recude"),
+			ERDGPassFlags::Compute | ERDGPassFlags::NeverCull,
+			Shader,
+			ShaderParameters,
+			FIntVector(NumReduceGroups, 1, 1));
+	}
+
+	{
+		FAutoExposureReduceFinalCS::FParameters *ShaderParameters = GraphBuilder.AllocParameters<FAutoExposureReduceFinalCS::FParameters>();
+		ShaderParameters->InputSize = NumReduceGroups;
+		ShaderParameters->InputSums = GraphBuilder.CreateSRV(FRDGBufferSRVDesc(OutputSums, EPixelFormat::PF_R32_FLOAT));
+		ShaderParameters->InputCounts = GraphBuilder.CreateSRV(FRDGBufferSRVDesc(OutputCounts, EPixelFormat::PF_R32_SINT));
+		ShaderParameters->OutputBuffer = GraphBuilder.CreateUAV(FRDGBufferUAVDesc(OutputBuffer, EPixelFormat::PF_R32_FLOAT));
+
+		FGlobalShaderMap* GlobalShaderMap = GetGlobalShaderMap(GMaxRHIFeatureLevel);
+		TShaderMapRef<FAutoExposureReduceFinalCS> Shader(GlobalShaderMap);
+
+		RDG_EVENT_SCOPE(GraphBuilder, "AutoExposure.RecudeFinal");
+		RDG_GPU_STAT_SCOPE(GraphBuilder, FAutoExposureReduceFinal);
+
+		FComputeShaderUtils::AddPass(
+			GraphBuilder,
+			RDG_EVENT_NAME("AutoExposure.RecudeFinal"),
+			ERDGPassFlags::Compute | ERDGPassFlags::NeverCull,
+			Shader,
+			ShaderParameters,
+			FIntVector(1, 1, 1));
+	}
+}
+
+}
\ No newline at end of file
diff --git a/Engine/Plugins/NNE/NNEDenoiser/Source/NNEDenoiser/Private/NNEDenoiserAutoExposure.h b/Engine/Plugins/NNE/NNEDenoiser/Source/NNEDenoiser/Private/NNEDenoiserAutoExposure.h
new file mode 100644
--- /dev/null
+++ b/Engine/Plugins/NNE/NNEDenoiser/Source/NNEDenoiser/Private/NNEDenoiserAutoExposure.h
@@ -0,0 +1,32 @@
+// Copyright Epic Games, Inc. All Rights Reserved.
+
+#pragma once
+
+#include "Containers/ContainersFwd.h"
+#include "Math/MathFwd.h"
+#include "RenderGraphFwd.h"
+
+struct FLinearColor;
+
+namespace UE::NNEDenoiser::Private
+{
+
+class IAutoExposure
+{
+public:
+	virtual ~IAutoExposure() = default;
+
+	virtual void Run(TConstArrayView<FLinearColor> InputData, FIntPoint Size, float& OutputValue) const = 0;
+	virtual void EnqueueRDG(FRDGBuilder& GraphBuilder, FRDGTextureRef InputTexture, FRDGBufferRef OutputBuffer) const = 0;
+};
+
+class FAutoExposure : public IAutoExposure
+{
+public:
+	virtual ~FAutoExposure() = default;
+
+	virtual void Run(TConstArrayView<FLinearColor> InputData, FIntPoint Size, float& OutputValue) const override;
+	virtual void EnqueueRDG(FRDGBuilder& GraphBuilder, FRDGTextureRef InputTexture, FRDGBufferRef OutputBuffer) const override;
+};
+
+}
\ No newline at end of file
diff --git a/Engine/Plugins/NNE/NNEDenoiser/Source/NNEDenoiser/Private/NNEDenoiserGenericDenoiser.cpp b/Engine/Plugins/NNE/NNEDenoiser/Source/NNEDenoiser/Private/NNEDenoiserGenericDenoiser.cpp
new file mode 100644
--- /dev/null
+++ b/Engine/Plugins/NNE/NNEDenoiser/Source/NNEDenoiser/Private/NNEDenoiserGenericDenoiser.cpp
@@ -0,0 +1,282 @@
+// Copyright Epic Games, Inc. All Rights Reserved.
+
+#include "NNEDenoiserGenericDenoiser.h"
+#include "Algo/Transform.h"
+#include "NNEDenoiserAutoExposure.h"
+#include "NNEDenoiserHistory.h"
+#include "NNEDenoiserIOProcess.h"
+#include "NNEDenoiserLog.h"
+#include "NNEDenoiserModelInstance.h"
+#include "NNEDenoiserResourceManager.h"
+#include "NNEDenoiserTransferFunction.h"
+#include "NNEDenoiserUtils.h"
+#include <type_traits>
+
+namespace UE::NNEDenoiser::Private
+{
+
+class FResourceAccess : public IResourceAccess
+{
+public:
+	FResourceAccess(const FResourceManager& ResourceManager) : ResourceManager(ResourceManager) { }
+
+	virtual FRDGTextureRef GetTexture(EResourceName TensorName, int32 FrameIdx) const override
+	{
+		return ResourceManager.GetTexture(TensorName, FrameIdx);
+	}
+
+	virtual FRDGTextureRef GetIntermediateTexture(EResourceName TensorName, int32 FrameIdx) const override
+	{
+		return ResourceManager.GetIntermediateTexture(TensorName, FrameIdx);
+	}
+
+private:
+	const FResourceManager& ResourceManager;
+};
+
+FRDGBufferRef CreateBufferRDG(FRDGBuilder& GraphBuilder, const NNE::FTensorDesc& TensorDesc, const NNE::FTensorShape& TensorShape)
+{
+	const uint32 BytesPerElement = TensorDesc.GetElementByteSize();
+	const uint32 NumElements = TensorShape.Volume();
+
+	FRDGBufferDesc Desc = FRDGBufferDesc::CreateBufferDesc(BytesPerElement, NumElements);
+	Desc.Usage |= EBufferUsageFlags::NNE;
+
+	return GraphBuilder.CreateBuffer(Desc, *TensorDesc.GetName());
+}
+
+TArray<FRDGBufferRef> CreateBuffersRDG(FRDGBuilder& GraphBuilder, TConstArrayView<NNE::FTensorDesc> TensorDescs, TConstArrayView<NNE::FTensorShape> TensorShapes)
+{
+	check(TensorDescs.Num() == TensorShapes.Num());
+
+	TArray<FRDGBufferRef> Result;
+	Result.SetNumUninitialized(TensorDescs.Num());
+
+	for (int32 I = 0; I < TensorDescs.Num(); I++)
+	{
+		Result[I] = CreateBufferRDG(GraphBuilder, TensorDescs[I], TensorShapes[I]);
+	}
+	return Result;
+}
+
+TArray<NNE::FTensorBindingRDG> GetBindingRDG(TConstArrayView<FRDGBufferRef> Buffers)
+{
+	TArray<NNE::FTensorBindingRDG> Result;
+	Result.Reserve(Buffers.Num());
+	
+	Algo::Transform(Buffers, Result, [] (FRDGBuffer* Buffer) { return NNE::FTensorBindingRDG{Buffer}; });
+
+	return Result;
+}
+
+void RegisterTextureIfNeeded(FResourceManager& ResourceManager, FRDGTextureRef Texture, EResourceName ResourceName, const IInputProcess& InputProcess)
+{
+	const int32 MinNumFrames = ResourceName == EResourceName::Output ? 1 : 0;
+	const int32 NumFrames = FMath::Max(InputProcess.NumFrames(ResourceName), MinNumFrames);
+	if (NumFrames <= 0)
+	{
+		return;
+	}
+	
+	ResourceManager.AddTexture(ResourceName, Texture, NumFrames);
+}
+
+TArray<NNE::FTensorShape> GetOutputTensorShapes(IModelInstance& ModelInstance)
+{
+	TConstArrayView<NNE::FTensorShape> InputShapes = ModelInstance.GetInputTensorShapes();
+	check(!InputShapes.IsEmpty());
+
+	TConstArrayView<uint32> InputShapeData = InputShapes[0].GetData();
+	const uint32 Width = InputShapeData[3];
+	const uint32 Height = InputShapeData[2];
+
+	TConstArrayView<NNE::FTensorDesc> OutputDescs = ModelInstance.GetOutputTensorDescs();
+	TConstArrayView<NNE::FTensorShape> OutputShapes = ModelInstance.GetOutputTensorShapes();
+
+	if (OutputDescs.Num() == OutputShapes.Num())
+	{
+		return TArray<NNE::FTensorShape>(OutputShapes);
+	}
+
+	// If output shapes not set yet, try to manually resolve, otherwise they need to be user specified.
+	TArray<NNE::FTensorShape> Result;
+	for (int32 i = 0; i < OutputDescs.Num(); i++)
+	{
+		check(OutputDescs[i].GetShape().Rank() == 4);
+
+		TConstArrayView<int32> SymbolicShapeData = OutputDescs[i].GetShape().GetData();
+		
+		TArray<uint32> ShapeData;
+		ShapeData.SetNumUninitialized(OutputDescs[i].GetShape().Rank());
+		ShapeData[0] = (uint32)SymbolicShapeData[0];
+		ShapeData[1] = (uint32)SymbolicShapeData[1];
+		ShapeData[2] = (uint32)Height;
+		ShapeData[3] = (uint32)Width;
+
+		Result.Add(NNE::FTensorShape::Make(ShapeData));
+	}
+	return Result;
+}
+
+void AddTilePasses(FRDGBuilder& GraphBuilder, IModelInstance& ModelInstance, const IInputProcess& InputProcess, const IOutputProcess& OutputProcess,
+		FResourceManager& ResourceManager, TConstArrayView<FRDGBufferRef> InputBuffers, TConstArrayView<FRDGBufferRef> OutputBuffers)
+{
+	FResourceAccess ResourceAccess(ResourceManager);
+
+	// 1. Read input and write to input buffers
+	InputProcess.AddPasses(GraphBuilder, ModelInstance.GetInputTensorDescs(), ModelInstance.GetInputTensorShapes(), ResourceAccess, InputBuffers);
+
+	// 2. Create buffer binding and infer the model
+	NNE::IModelInstanceRDG::EEnqueueRDGStatus Status = ModelInstance.EnqueueRDG(GraphBuilder, GetBindingRDG(InputBuffers), GetBindingRDG(OutputBuffers));
+	checkf(Status == NNE::IModelInstanceRDG::EEnqueueRDGStatus::Ok, TEXT("EnqueueRDG failed: %d"), static_cast<int>(Status));
+
+	// 3. Write output based on output buffer
+	FRDGTextureRef OutputTexture = ResourceManager.GetTexture(EResourceName::Output, 0);
+	OutputProcess.AddPasses(GraphBuilder, ModelInstance.GetOutputTensorDescs(), GetOutputTensorShapes(ModelInstance), ResourceAccess, OutputBuffers, OutputTexture);
+}
+
+FGenericDenoiser::FGenericDenoiser(
+	TUniquePtr<IModelInstance> ModelInstance,
+	TUniquePtr<IInputProcess> InputProcess,
+	TUniquePtr<IOutputProcess> OutputProcess,
+	FParameters DenoiserParameters,
+	TUniquePtr<IAutoExposure> AutoExposure,
+	TSharedPtr<ITransferFunction> TransferFunction) :
+		ModelInstance(MoveTemp(ModelInstance)),
+		InputProcess(MoveTemp(InputProcess)),
+		OutputProcess(MoveTemp(OutputProcess)),
+		DenoiserParameters(DenoiserParameters),
+		AutoExposure(MoveTemp(AutoExposure)),
+		TransferFunction(TransferFunction)
+{
+
+}
+
+FGenericDenoiser::~FGenericDenoiser()
+{
+
+}
+
+bool FGenericDenoiser::Prepare(FIntPoint Extent)
+{
+	if (Extent == LastExtent)
+	{
+		return true;
+	}
+
+	// Probably would be enough to do this only once at the very beginning...
+	// We just want to be sure that everything up to width and height is correct
+	if (!InputProcess->Validate(*ModelInstance, {-1, -1}))
+	{
+		return false;
+	}
+
+	TConstArrayView<int32> SymbolicInputShape = ModelInstance->GetInputTensorDescs()[0].GetShape().GetData();
+	const FIntPoint TargetTileSize = {SymbolicInputShape[3], SymbolicInputShape[2]};
+
+	Tiling = CreateTiling(TargetTileSize,
+				DenoiserParameters.TilingConfig.MaxSize,
+				DenoiserParameters.TilingConfig.MinSize,
+				DenoiserParameters.TilingConfig.Alignment,
+				DenoiserParameters.TilingConfig.Overlap,
+				Extent);
+
+	if (!InputProcess->Prepare(*ModelInstance, Tiling.TileSize))
+	{
+		return false;
+	}
+
+	LastExtent = Extent;
+
+	UE_LOG(LogNNEDenoiser, Log, TEXT("Prepared neural denoiser model:\n  Viewport   %dx%d\n  Num. tiles %dx%d\n  Tile size  %dx%d"),
+		Extent.X, Extent.Y, Tiling.Count.X, Tiling.Count.Y, Tiling.TileSize.X, Tiling.TileSize.Y);
+
+	return true;
+}
+
+TUniquePtr<FHistory> FGenericDenoiser::AddPasses(
+	FRDGBuilder& GraphBuilder,
+	FRDGTextureRef ColorTex,
+	FRDGTextureRef AlbedoTex,
+	FRDGTextureRef NormalTex,
+	FRDGTextureRef OutputTex,
+	FRDGTextureRef FlowTex,
+	const FRHIGPUMask& GPUMask,
+	FHistory* History)
+{
+	using namespace UE::NNEDenoiserShaders::Internal;
+
+	SCOPED_NAMED_EVENT_TEXT("NNEDenoiser.AddPasses", FColor::Magenta);
+
+	// note: FlowTex is currently optional and can be null
+	check(ColorTex != FRDGTextureRef{});
+	check(AlbedoTex != FRDGTextureRef{});
+	check(NormalTex != FRDGTextureRef{});
+	check(OutputTex != FRDGTextureRef{});
+
+	const FIntPoint Extent = ColorTex->Desc.Extent;
+
+	TUniquePtr<FHistory> Result;
+
+	if (!Prepare(Extent))
+	{
+		AddCopyTexturePass(GraphBuilder, ColorTex, OutputTex, FRHICopyTextureInfo{});
+
+		return Result;
+	}
+
+	TMap<EResourceName, TArray<TRefCountPtr<IPooledRenderTarget>>> ResourceMap;
+	if (History)
+	{
+		ResourceMap = History->GetResourceMap();
+	}
+
+	FResourceManager ResourceManager(GraphBuilder, Tiling, ResourceMap);
+	RegisterTextureIfNeeded(ResourceManager, ColorTex, EResourceName::Color, *InputProcess);
+	RegisterTextureIfNeeded(ResourceManager, AlbedoTex, EResourceName::Albedo, *InputProcess);
+	RegisterTextureIfNeeded(ResourceManager, NormalTex, EResourceName::Normal, *InputProcess);
+	if (FlowTex != FRDGTextureRef{})
+	{
+		RegisterTextureIfNeeded(ResourceManager, FlowTex, EResourceName::Flow, *InputProcess);
+	}
+	RegisterTextureIfNeeded(ResourceManager, OutputTex, EResourceName::Output, *InputProcess);
+
+	TArray<FRDGBufferRef> InputBuffers = CreateBuffersRDG(GraphBuilder, ModelInstance->GetInputTensorDescs(), ModelInstance->GetInputTensorShapes());
+	TArray<FRDGBufferRef> OutputBuffers = CreateBuffersRDG(GraphBuilder, ModelInstance->GetOutputTensorDescs(), GetOutputTensorShapes(*ModelInstance));
+
+	if (TransferFunction.IsValid() && AutoExposure.IsValid())
+	{
+		FRDGBufferDesc InputBufferDesc = FRDGBufferDesc::CreateBufferDesc(sizeof(float), 2);
+		InputBufferDesc.Usage |= EBufferUsageFlags::NNE;
+
+		FRDGBufferRef InputScaleBuffer = GraphBuilder.CreateBuffer(InputBufferDesc, TEXT("AutoExposureOutputBuffer"));
+
+		AutoExposure->EnqueueRDG(GraphBuilder, ColorTex, InputScaleBuffer);
+
+		TransferFunction->RDGSetInputScale(InputScaleBuffer);
+	}
+	else
+	{
+		checkf(!TransferFunction.IsValid() && !AutoExposure.IsValid(), TEXT("TransferFunction and AutoExposure either both need to be set or not set."))
+	}
+
+	for (int32 I = 0; I < Tiling.Tiles.Num(); I++)
+	{
+		ResourceManager.BeginTile(I);
+
+		// Do inference on tile
+		AddTilePasses(GraphBuilder, *ModelInstance, *InputProcess, *OutputProcess, ResourceManager, InputBuffers, OutputBuffers);
+
+		ResourceManager.EndTile();
+	}
+
+	ResourceMap = ResourceManager.MakeHistoryResourceMap();
+	if (!ResourceMap.IsEmpty())
+	{
+		Result = MakeUnique<FHistory>(*DebugName, MoveTemp(ResourceMap));
+	}
+
+	return Result;
+}
+
+} // namespace UE::NNEDenoiser::Private
\ No newline at end of file
diff --git a/Engine/Plugins/NNE/NNEDenoiser/Source/NNEDenoiser/Private/NNEDenoiserGenericDenoiser.h b/Engine/Plugins/NNE/NNEDenoiser/Source/NNEDenoiser/Private/NNEDenoiserGenericDenoiser.h
new file mode 100644
--- /dev/null
+++ b/Engine/Plugins/NNE/NNEDenoiser/Source/NNEDenoiser/Private/NNEDenoiserGenericDenoiser.h
@@ -0,0 +1,67 @@
+// Copyright Epic Games, Inc. All Rights Reserved.
+
+#pragma once
+
+#include "Containers/UnrealString.h"
+#include "Math/IntPoint.h"
+#include "NNEDenoiserParameters.h"
+#include "NNEDenoiserTiling.h"
+#include "RenderGraphFwd.h"
+#include "Templates/UniquePtr.h"
+
+struct FRHIGPUMask;
+
+namespace UE::NNEDenoiser::Private
+{
+
+class FHistory;
+class IModelInstance;
+class IInputProcess;
+class IOutputProcess;
+class IAutoExposure;
+class ITransferFunction;
+
+class FGenericDenoiser
+{
+public:
+	FGenericDenoiser(
+		TUniquePtr<IModelInstance> ModelInstance,
+		TUniquePtr<IInputProcess> InputProcess,
+		TUniquePtr<IOutputProcess> OutputProcess,
+		FParameters DenoiserParameters,
+		TUniquePtr<IAutoExposure> AutoExposure,
+		TSharedPtr<ITransferFunction> TransferFunction
+	);
+
+	~FGenericDenoiser();
+
+	// IPathTracingSpatialTemporalDenoiser interface
+	const TCHAR* GetDebugName() const { return *DebugName; }
+
+	TUniquePtr<FHistory> AddPasses(
+		FRDGBuilder &GraphBuilder,
+		FRDGTextureRef ColorTex,
+		FRDGTextureRef AlbedoTex,
+		FRDGTextureRef NormalTex,
+		FRDGTextureRef OutputTex,
+		FRDGTextureRef FlowTex,
+		const FRHIGPUMask& GPUMask,
+		FHistory* History);
+
+private:
+	bool Prepare(FIntPoint Extent);
+
+	inline static const FString DebugName = TEXT("FGenericDenoiser");
+
+	TUniquePtr<IModelInstance> ModelInstance;
+	TUniquePtr<IInputProcess> InputProcess;
+	TUniquePtr<IOutputProcess> OutputProcess;
+	FParameters DenoiserParameters;
+	TUniquePtr<IAutoExposure> AutoExposure;
+	TSharedPtr<ITransferFunction> TransferFunction;
+
+	FIntPoint LastExtent = { -1, -1 };
+	FTiling Tiling{};
+};
+
+} // namespace UE::NNEDenoiser::Private
\ No newline at end of file
diff --git a/Engine/Plugins/NNE/NNEDenoiser/Source/NNEDenoiser/Private/NNEDenoiserHistory.h b/Engine/Plugins/NNE/NNEDenoiser/Source/NNEDenoiser/Private/NNEDenoiserHistory.h
new file mode 100644
--- /dev/null
+++ b/Engine/Plugins/NNE/NNEDenoiser/Source/NNEDenoiser/Private/NNEDenoiserHistory.h
@@ -0,0 +1,36 @@
+// Copyright Epic Games, Inc. All Rights Reserved.
+
+#pragma once
+
+#include "CoreMinimal.h"
+#include "NNEDenoiserLog.h"
+#include "NNEDenoiserResourceName.h"
+#include "PathTracingDenoiser.h"
+#include "RendererInterface.h"
+
+struct IPooledRenderTarget;
+
+namespace UE::NNEDenoiser::Private
+{
+
+using UE::Renderer::Private::IPathTracingSpatialTemporalDenoiser;
+
+class FHistory : public IPathTracingSpatialTemporalDenoiser::IHistory
+{
+public:
+	FHistory(const TCHAR* DebugName, TMap<EResourceName, TArray<TRefCountPtr<IPooledRenderTarget>>> ResourceMap) : DebugName(DebugName), ResourceMap(MoveTemp(ResourceMap)) {
+
+	}
+
+	virtual ~FHistory() = default;
+
+	const TCHAR* GetDebugName() const override { return DebugName; };
+
+	const TMap<EResourceName, TArray<TRefCountPtr<IPooledRenderTarget>>>& GetResourceMap() const { return ResourceMap; }
+
+private:
+	const TCHAR* DebugName;
+	TMap<EResourceName, TArray<TRefCountPtr<IPooledRenderTarget>>> ResourceMap;
+};
+
+}
\ No newline at end of file
diff --git a/Engine/Plugins/NNE/NNEDenoiser/Source/NNEDenoiser/Private/NNEDenoiserIOMappingData.cpp b/Engine/Plugins/NNE/NNEDenoiser/Source/NNEDenoiser/Private/NNEDenoiserIOMappingData.cpp
new file mode 100644
--- /dev/null
+++ b/Engine/Plugins/NNE/NNEDenoiser/Source/NNEDenoiser/Private/NNEDenoiserIOMappingData.cpp
@@ -0,0 +1,55 @@
+// Copyright Epic Games, Inc. All Rights Reserved.
+
+#include "NNEDenoiserIOMappingData.h"
+
+namespace UE::NNEDenoiser
+{
+	
+EResourceName ToResourceName(EInputResourceName Name)
+{
+	switch(Name)
+	{
+		case EInputResourceName::Color: return EResourceName::Color;
+		case EInputResourceName::Albedo: return EResourceName::Albedo;
+		case EInputResourceName::Normal: return EResourceName::Normal;
+		case EInputResourceName::Output: return EResourceName::Output;
+		default: checkNoEntry();
+	}
+	return EResourceName::Color;
+}
+
+EResourceName ToResourceName(EOutputResourceName Name)
+{
+	switch(Name)
+	{
+		case EOutputResourceName::Output: return EResourceName::Output;
+		default: checkNoEntry();
+	}
+	return EResourceName::Color;
+}
+
+EResourceName ToResourceName(ETemporalInputResourceName Name)
+{
+	switch(Name)
+	{
+		case ETemporalInputResourceName::Color: return EResourceName::Color;
+		case ETemporalInputResourceName::Albedo: return EResourceName::Albedo;
+		case ETemporalInputResourceName::Normal: return EResourceName::Normal;
+		case ETemporalInputResourceName::Flow: return EResourceName::Flow;
+		case ETemporalInputResourceName::Output: return EResourceName::Output;
+		default: checkNoEntry();
+	}
+	return EResourceName::Color;
+}
+
+EResourceName ToResourceName(ETemporalOutputResourceName Name)
+{
+	switch(Name)
+	{
+		case ETemporalOutputResourceName::Output: return EResourceName::Output;
+		default: checkNoEntry();
+	}
+	return EResourceName::Color;
+}
+
+} // namespace UE::NNEDenoiser
\ No newline at end of file
diff --git a/Engine/Plugins/NNE/NNEDenoiser/Source/NNEDenoiser/Private/NNEDenoiserIOProcess.h b/Engine/Plugins/NNE/NNEDenoiser/Source/NNEDenoiser/Private/NNEDenoiserIOProcess.h
new file mode 100644
--- /dev/null
+++ b/Engine/Plugins/NNE/NNEDenoiser/Source/NNEDenoiser/Private/NNEDenoiserIOProcess.h
@@ -0,0 +1,57 @@
+// Copyright Epic Games, Inc. All Rights Reserved.
+
+#pragma once
+
+#include "NNEDenoiserResourceName.h"
+#include "NNETypes.h"
+#include "RenderGraphFwd.h"
+
+namespace UE::NNEDenoiser::Private
+{
+
+class FResourceMapping;
+class IModelInstance;
+
+class IResourceAccess
+{
+public:
+	virtual FRDGTextureRef GetTexture(EResourceName TensorName, int32 FrameIdx) const = 0;
+	virtual FRDGTextureRef GetIntermediateTexture(EResourceName TensorName, int32 FrameIdx) const = 0;
+};
+
+class IInputProcess
+{
+public:
+	virtual ~IInputProcess() = default;
+
+	virtual bool Validate(const IModelInstance& ModelInstance, FIntPoint Extent) const = 0;
+
+	virtual bool Prepare(IModelInstance& ModelInstance, FIntPoint Extent) const = 0;
+
+	virtual int32 NumFrames(EResourceName Name) const = 0;
+
+	virtual void AddPasses(
+		FRDGBuilder& GraphBuilder,
+		TConstArrayView<UE::NNE::FTensorDesc> TensorDescs,
+		TConstArrayView<UE::NNE::FTensorShape> TensorShapes,
+		const IResourceAccess& ResourceAccess,
+		TConstArrayView<FRDGBufferRef> OutputBuffers) const = 0;
+};
+
+class IOutputProcess
+{
+public:
+	virtual ~IOutputProcess() = default;
+
+	virtual bool Validate(const IModelInstance& ModelInstance, FIntPoint Extent) const = 0;
+
+	virtual void AddPasses(
+		FRDGBuilder& GraphBuilder,
+		TConstArrayView<UE::NNE::FTensorDesc> TensorDescs,
+		TConstArrayView<UE::NNE::FTensorShape> TensorShapes,
+		const IResourceAccess& ResourceAccess,
+		TConstArrayView<FRDGBufferRef> Buffers,
+		FRDGTextureRef OutputTexture) const = 0;
+};
+
+} // namespace UE::NNEDenoiser::Private
\ No newline at end of file
diff --git a/Engine/Plugins/NNE/NNEDenoiser/Source/NNEDenoiser/Private/NNEDenoiserIOProcessBase.cpp b/Engine/Plugins/NNE/NNEDenoiser/Source/NNEDenoiser/Private/NNEDenoiserIOProcessBase.cpp
new file mode 100644
--- /dev/null
+++ b/Engine/Plugins/NNE/NNEDenoiser/Source/NNEDenoiser/Private/NNEDenoiserIOProcessBase.cpp
@@ -0,0 +1,523 @@
+// Copyright Epic Games, Inc. All Rights Reserved.
+
+#include "NNEDenoiserIOProcessBase.h"
+#include "NNEDenoiserLog.h"
+#include "NNEDenoiserModelInstance.h"
+#include "NNEDenoiserParameters.h"
+#include "NNEDenoiserShadersMappedCopyCS.h"
+#include "NNEDenoiserShadersDefaultIOProcessCS.h"
+#include "NNEDenoiserTransferFunction.h"
+#include "NNEDenoiserUtils.h"
+#include "RHIStaticStates.h"
+
+DECLARE_GPU_STAT_NAMED(FNNEDenoiserReadInput, TEXT("NNEDenoiser.ReadInput"));
+DECLARE_GPU_STAT_NAMED(FNNEDenoiserWriteOutput, TEXT("NNEDenoiser.WriteOutput"));
+DECLARE_GPU_STAT_NAMED(FNNEDenoiserDefaultIOProcess, TEXT("NNEDenoiser.DefaultIOProcess"));
+
+namespace UE::NNEDenoiser::Private
+{
+
+namespace IOProcessBaseHelper
+{
+
+NNEDenoiserShaders::Internal::EDefaultIOProcessInputKind GetInputKind(EResourceName TensorName)
+{
+	using NNEDenoiserShaders::Internal::EDefaultIOProcessInputKind;
+	
+	switch(TensorName)
+	{
+		case EResourceName::Color:	return EDefaultIOProcessInputKind::Color;
+		case EResourceName::Albedo:	return EDefaultIOProcessInputKind::Albedo;
+		case EResourceName::Normal:	return EDefaultIOProcessInputKind::Normal;
+		case EResourceName::Flow:	return EDefaultIOProcessInputKind::Flow;
+		case EResourceName::Output:	return EDefaultIOProcessInputKind::Output;
+	}
+
+	// There should be a case for every resource name
+	checkNoEntry();
+	
+	return EDefaultIOProcessInputKind::Color;
+}
+
+void AddPreOrPostProcess(
+	FRDGBuilder& GraphBuilder,
+	FRDGTextureRef InputTexture,
+	EResourceName TensorName,
+	int32 FrameIdx,
+	FRDGTextureRef OutputTexture)
+{
+	using namespace UE::NNEDenoiserShaders::Internal;
+
+	const FIntVector Size = InputTexture->Desc.GetSize();
+	check(Size == OutputTexture->Desc.GetSize());
+
+	FDefaultIOProcessCS::FParameters *ShaderParameters = GraphBuilder.AllocParameters<FDefaultIOProcessCS::FParameters>();
+	ShaderParameters->Width = Size.X;
+	ShaderParameters->Height = Size.Y;
+	ShaderParameters->InputTexture = InputTexture;
+	ShaderParameters->OutputTexture = GraphBuilder.CreateUAV(FRDGTextureUAVDesc(OutputTexture));
+
+	FDefaultIOProcessCS::FPermutationDomain PermutationVector;
+	PermutationVector.Set<FDefaultIOProcessCS::FDefaultIOProcessInputKind>(GetInputKind(TensorName));
+
+	FIntVector ThreadGroupCount = FIntVector(
+		FMath::DivideAndRoundUp(Size.X, FDefaultIOProcessConstants::THREAD_GROUP_SIZE),
+		FMath::DivideAndRoundUp(Size.Y, FDefaultIOProcessConstants::THREAD_GROUP_SIZE),
+		1);
+
+	FGlobalShaderMap* GlobalShaderMap = GetGlobalShaderMap(GMaxRHIFeatureLevel);
+	TShaderMapRef<FDefaultIOProcessCS> Shader(GlobalShaderMap, PermutationVector);
+
+	RDG_EVENT_SCOPE(GraphBuilder, "NNEDenoiser.DefaultIOProcess");
+	RDG_GPU_STAT_SCOPE(GraphBuilder, FNNEDenoiserDefaultIOProcess);
+
+	FComputeShaderUtils::AddPass(
+		GraphBuilder,
+		RDG_EVENT_NAME("NNEDenoiser.DefaultIOProcess"),
+		ERDGPassFlags::Compute | ERDGPassFlags::NeverCull,
+		Shader,
+		ShaderParameters,
+		ThreadGroupCount);
+}
+
+void AddReadInputPass(
+	FRDGBuilder& GraphBuilder,
+	FRDGTextureRef InputTexture,
+	FRDGBufferUAVRef OutputBufferUAV,
+	const TArray<FChannelMapping>& ChannelMapping)
+{
+	using namespace UE::NNEDenoiserShaders::Internal;
+
+	if (ChannelMapping.IsEmpty())
+	{
+		UE_LOG(LogNNEDenoiser, Warning, TEXT("AddReadInputPass: ChannelMapping is empty. Nothing to do!"));
+		return;
+	}
+	
+	if (ChannelMapping.Num() >= FMappedCopyConstants::MAX_NUM_MAPPED_CHANNELS)
+	{
+		UE_LOG(LogNNEDenoiser, Warning, TEXT("AddReadInputPass: ChannelMapping has too many entries!"));
+		return;
+	}
+
+	const FIntVector Size = InputTexture->Desc.GetSize();
+
+	FTextureBufferMappedCopyCS::FParameters *ReadInputParameters = GraphBuilder.AllocParameters<FTextureBufferMappedCopyCS::FParameters>();
+	ReadInputParameters->Width = Size.X;
+	ReadInputParameters->Height = Size.Y;
+	ReadInputParameters->InputTexture = InputTexture;
+	ReadInputParameters->OutputBuffer = OutputBufferUAV;
+	for (int32 Idx = 0; Idx < ChannelMapping.Num(); Idx++)
+	{
+		ReadInputParameters->OutputChannel_InputChannel_Unused_Unused[Idx] = { ChannelMapping[Idx].TensorChannel, ChannelMapping[Idx].ResourceChannel, 0, 0 };
+	}
+
+	const EDataType InputDataType = GetDenoiserShaderDataType(InputTexture->Desc.Format);
+	const EDataType OutputDataType = GetDenoiserShaderDataType(OutputBufferUAV->Desc.Format);
+
+	FTextureBufferMappedCopyCS::FPermutationDomain PermutationVector;
+	PermutationVector.Set<FTextureBufferMappedCopyCS::FInputDataType>(InputDataType);
+	PermutationVector.Set<FTextureBufferMappedCopyCS::FOutputDataType>(OutputDataType);
+	PermutationVector.Set<FTextureBufferMappedCopyCS::FNumMappedChannels>(ChannelMapping.Num());
+
+	FIntVector ReadInputThreadGroupCount = FIntVector(
+		FMath::DivideAndRoundUp(Size.X, FMappedCopyConstants::THREAD_GROUP_SIZE),
+		FMath::DivideAndRoundUp(Size.Y, FMappedCopyConstants::THREAD_GROUP_SIZE),
+		1);
+
+	FGlobalShaderMap* GlobalShaderMap = GetGlobalShaderMap(GMaxRHIFeatureLevel);
+	TShaderMapRef<FTextureBufferMappedCopyCS> ReadInputShader(GlobalShaderMap, PermutationVector);
+
+	RDG_EVENT_SCOPE(GraphBuilder, "NNEDenoiser.ReadInput");
+	RDG_GPU_STAT_SCOPE(GraphBuilder, FNNEDenoiserReadInput);
+
+	FComputeShaderUtils::AddPass(
+		GraphBuilder,
+		RDG_EVENT_NAME("NNEDenoiser.ReadInput"),
+		ERDGPassFlags::Compute | ERDGPassFlags::NeverCull,
+		ReadInputShader,
+		ReadInputParameters,
+		ReadInputThreadGroupCount);
+}
+
+void AddWriteOutputPass(
+	FRDGBuilder& GraphBuilder,
+	FRDGBufferRef InputBuffer,
+	EPixelFormat BufferFormat,
+	FRDGTextureRef OutputTexture,
+	UE::NNEDenoiserShaders::Internal::EDataType DataType,
+	const TArray<FChannelMapping>& ChannelMapping)
+{
+	using namespace UE::NNEDenoiserShaders::Internal;
+
+	const FIntVector Size = OutputTexture->Desc.GetSize();
+
+	FBufferTextureMappedCopyCS::FParameters *WriteOutputParameters = GraphBuilder.AllocParameters<FBufferTextureMappedCopyCS::FParameters>();
+	WriteOutputParameters->Width = Size.X;
+	WriteOutputParameters->Height = Size.Y;
+	WriteOutputParameters->InputBuffer = GraphBuilder.CreateSRV(FRDGBufferSRVDesc(InputBuffer, BufferFormat));
+	WriteOutputParameters->OutputTexture = GraphBuilder.CreateUAV(FRDGTextureUAVDesc(OutputTexture));
+	for (int32 Idx = 0; Idx < ChannelMapping.Num(); Idx++)
+	{
+		WriteOutputParameters->OutputChannel_InputChannel_Unused_Unused[Idx] = { ChannelMapping[Idx].ResourceChannel, ChannelMapping[Idx].TensorChannel, 0, 0 };
+	}
+
+	const EDataType InputDataType = GetDenoiserShaderDataType(BufferFormat);
+	const EDataType OutputDataType = GetDenoiserShaderDataType(OutputTexture->Desc.Format);
+
+	FBufferTextureMappedCopyCS::FPermutationDomain PermutationVector;
+	PermutationVector.Set<FBufferTextureMappedCopyCS::FInputDataType>(InputDataType);
+	PermutationVector.Set<FBufferTextureMappedCopyCS::FOutputDataType>(OutputDataType);
+	PermutationVector.Set<FBufferTextureMappedCopyCS::FNumMappedChannels>(ChannelMapping.Num());
+
+	FIntVector WriteOutputThreadGroupCount = FIntVector(
+		FMath::DivideAndRoundUp(Size.X, FMappedCopyConstants::THREAD_GROUP_SIZE),
+		FMath::DivideAndRoundUp(Size.Y, FMappedCopyConstants::THREAD_GROUP_SIZE),
+		1);
+
+	FGlobalShaderMap* GlobalShaderMap = GetGlobalShaderMap(GMaxRHIFeatureLevel);
+	TShaderMapRef<FBufferTextureMappedCopyCS> WriteOutputShader(GlobalShaderMap, PermutationVector);
+
+	RDG_EVENT_SCOPE(GraphBuilder, "NNEDenoiser.WriteOutput");
+	RDG_GPU_STAT_SCOPE(GraphBuilder, FNNEDenoiserWriteOutput);
+
+	FComputeShaderUtils::AddPass(
+		GraphBuilder,
+		RDG_EVENT_NAME("NNEDenoiser.WriteOutput"),
+		ERDGPassFlags::Compute | ERDGPassFlags::NeverCull,
+		WriteOutputShader,
+		WriteOutputParameters,
+		WriteOutputThreadGroupCount);
+}
+
+void AddReadInputPassForKind(
+	FRDGBuilder& GraphBuilder,
+	const IResourceAccess& ResourceAccess,
+	EResourceName TensorName,
+	const FResourceMapping& ResourceMapping,
+	FRDGBufferUAVRef BufferUAV)
+{
+	for (const TPair<int32, TArray<FChannelMapping>>& ChannelMappingKeyValue : ResourceMapping.GetChannelMappingPerFrame(TensorName))
+	{
+		const int32 FrameIdx = -ChannelMappingKeyValue.Key;
+		FRDGTextureRef InputTexture = ResourceAccess.GetTexture(TensorName, FrameIdx);
+
+		AddReadInputPass(GraphBuilder, InputTexture, BufferUAV, ChannelMappingKeyValue.Value);
+	}
+};
+
+} // namespace IOProcessBaseHelper
+
+bool FInputProcessBase::Validate(const IModelInstance& ModelInstance, FIntPoint Extent) const
+{
+	SCOPED_NAMED_EVENT_TEXT("NNEDenoiser.Validate", FColor::Magenta);
+
+	checkf(Extent == FIntPoint(-1, -1) || (Extent.X >= 0 && Extent.Y >= 0), TEXT("Extent should be either fully symbolic or set!"));
+
+	const int32 NumBatches = 1;
+
+	UE_LOG(LogNNEDenoiser, Log, TEXT("Validate model for extent %dx%d..."), Extent.X, Extent.Y);
+
+	if (Extent == FIntPoint(-1, -1))
+	{
+		TConstArrayView<UE::NNE::FTensorDesc> InputTensorDescs = ModelInstance.GetInputTensorDescs();
+		if (InputTensorDescs.Num() != InputLayout.Num())
+		{
+			UE_LOG(LogNNEDenoiser, Error, TEXT("Wrong number of inputs (expected %d, got %d)!"), InputLayout.Num(), InputTensorDescs.Num())
+			return false;
+		}
+
+		TArray<UE::NNE::FTensorShape> InputShapes;
+		for (int32 Idx = 0; Idx < InputTensorDescs.Num(); Idx++)
+		{
+			TConstArrayView<int32> InputSymbolicTensorShapeData = InputTensorDescs[Idx].GetShape().GetData();
+			const TArray<int32, TInlineAllocator<4>> RequiredInputShapeData = { NumBatches, InputLayout.NumChannels(Idx), Extent.Y, Extent.X };
+
+			if (!IsTensorShapeValid(InputSymbolicTensorShapeData, RequiredInputShapeData, TEXT("Input")))
+			{
+				return false;
+			}
+		}
+	}
+	else
+	{
+		TConstArrayView<UE::NNE::FTensorShape> InputTensorShapes = ModelInstance.GetInputTensorShapes();
+		if (InputTensorShapes.Num() != InputLayout.Num())
+		{
+			UE_LOG(LogNNEDenoiser, Error, TEXT("Wrong number of inputs (expected %d, got %d)!"), InputLayout.Num(), InputTensorShapes.Num())
+			return false;
+		}
+
+		for (int32 Idx = 0; Idx < InputTensorShapes.Num(); Idx++)
+		{
+			TConstArrayView<uint32> InputTensorShapeData = InputTensorShapes[Idx].GetData();
+			const TArray<int32, TInlineAllocator<4>> RequiredInputShapeData = { NumBatches, InputLayout.NumChannels(Idx), Extent.Y, Extent.X };
+
+			if (!IsTensorShapeValid(InputTensorShapeData, RequiredInputShapeData, TEXT("Input")))
+			{
+				return false;
+			}
+		}
+	}
+
+	return true;
+}
+
+bool FInputProcessBase::Prepare(IModelInstance& ModelInstance, FIntPoint Extent) const
+{
+	SCOPED_NAMED_EVENT_TEXT("NNEDenoiser.Prepare", FColor::Magenta);
+
+	if (!Validate(ModelInstance, {-1, -1}))
+	{
+		return false;
+	}
+
+	const int32 NumBatches = 1;
+
+	UE_LOG(LogNNEDenoiser, Log, TEXT("Configure model for extent %dx%d..."), Extent.X, Extent.Y);
+
+	TConstArrayView<UE::NNE::FTensorDesc> InputTensorDescs = ModelInstance.GetInputTensorDescs();
+	
+	UE_LOG(LogNNEDenoiser, Log, TEXT("Input shapes (set):"));
+
+	TArray<UE::NNE::FTensorShape> InputShapes;
+	for (int32 Idx = 0; Idx < InputTensorDescs.Num(); Idx++)
+	{
+		TConstArrayView<int32> InputSymbolicTensorShapeData = InputTensorDescs[Idx].GetShape().GetData();
+
+		const FIntVector4 ModelInputShape = { 1, InputSymbolicTensorShapeData[1], Extent.Y, Extent.X };
+
+		InputShapes.Add(UE::NNE::FTensorShape::Make({ (uint32)ModelInputShape.X, (uint32)ModelInputShape.Y, (uint32)ModelInputShape.Z, (uint32)ModelInputShape.W }));
+
+		UE_LOG(LogNNEDenoiser, Log, TEXT("%d: (%d, %d, %d, %d)"), Idx, ModelInputShape.X, ModelInputShape.Y, ModelInputShape.Z, ModelInputShape.W)
+	}
+
+	const IModelInstance::ESetInputTensorShapesStatus Status = ModelInstance.SetInputTensorShapes(InputShapes);
+	if (Status != IModelInstance::ESetInputTensorShapesStatus::Ok)
+	{
+		UE_LOG(LogNNEDenoiser, Error, TEXT("Could not configure model instance (ModelInstance.SetInputTensorShapes() failed)!"))
+		return false;
+	}
+
+	if (!Validate(ModelInstance, Extent))
+	{
+		return false;
+	}
+
+	return true;
+}
+
+int32 FInputProcessBase::NumFrames(EResourceName Name) const
+{
+	return InputLayout.NumFrames(Name);
+}
+
+void FInputProcessBase::AddPasses(
+		FRDGBuilder& GraphBuilder,
+		TConstArrayView<UE::NNE::FTensorDesc> TensorDescs,
+		TConstArrayView<UE::NNE::FTensorShape> TensorShapes,
+		const IResourceAccess& ResourceAccess,
+		TConstArrayView<FRDGBufferRef> OutputBuffers) const
+{
+	const UEnum* ResourceNameEnum = StaticEnum<EResourceName>();
+	for (int32 I = 0; I < ResourceNameEnum->NumEnums(); I++)
+	{
+		const EResourceName TensorName = (EResourceName)ResourceNameEnum->GetValueByIndex(I);
+		const int32 NumFrames = InputLayout.NumFrames(TensorName);
+
+		for (int32 J = 0; J < NumFrames; J++)
+		{
+			if (!HasPreprocessInput(TensorName, J))
+			{
+				continue;
+			}
+
+			FRDGTextureRef InputTexture = ResourceAccess.GetTexture(TensorName, J);
+			FRDGTextureRef PreprocessedInputTexture = ResourceAccess.GetIntermediateTexture(TensorName, J);
+
+			PreprocessInput(GraphBuilder, InputTexture, TensorName, J, PreprocessedInputTexture);
+		}
+	}
+
+	for (int32 I = 0; I < TensorDescs.Num(); I++)
+	{
+		WriteInputBuffer(GraphBuilder, TensorDescs[I], TensorShapes[I], ResourceAccess, InputLayout.GetChecked(I), OutputBuffers[I]);
+	}
+}
+
+bool FInputProcessBase::HasPreprocessInput(EResourceName TensorName, int32 FrameIdx) const
+{
+	check(FrameIdx == 0);
+
+	switch(TensorName)
+	{
+		case EResourceName::Color:
+		case EResourceName::Albedo:
+		case EResourceName::Normal:
+			return true;
+	}
+	return false;
+}
+
+void FInputProcessBase::PreprocessInput(
+		FRDGBuilder& GraphBuilder,
+		FRDGTextureRef Texture,
+		EResourceName TensorName,
+		int32 FrameIdx,
+		FRDGTextureRef PreprocessedTexture) const
+{
+	IOProcessBaseHelper::AddPreOrPostProcess(GraphBuilder, Texture, TensorName, FrameIdx, PreprocessedTexture);
+
+	if (TransferFunction.IsValid() && TensorName == EResourceName::Color)
+	{
+		TransferFunction->RDGForward(GraphBuilder, PreprocessedTexture, PreprocessedTexture);
+	}
+}
+
+void FInputProcessBase::WriteInputBuffer(
+	FRDGBuilder& GraphBuilder,
+	const UE::NNE::FTensorDesc& TensorDesc,
+	const UE::NNE::FTensorShape& TensorShape,
+	const IResourceAccess& ResourceAccess,
+	const FResourceMapping& ResourceMapping,
+	FRDGBufferRef Buffer) const
+{
+	const ENNETensorDataType TensorDataType = TensorDesc.GetDataType();
+
+	FRDGBufferUAVRef BufferUAV = GraphBuilder.CreateUAV(FRDGBufferUAVDesc(Buffer, GetBufferFormat(TensorDataType)));
+
+	IOProcessBaseHelper::AddReadInputPassForKind(GraphBuilder, ResourceAccess, EResourceName::Color, ResourceMapping, BufferUAV);
+	IOProcessBaseHelper::AddReadInputPassForKind(GraphBuilder, ResourceAccess, EResourceName::Albedo, ResourceMapping, BufferUAV);
+	IOProcessBaseHelper::AddReadInputPassForKind(GraphBuilder, ResourceAccess, EResourceName::Normal, ResourceMapping, BufferUAV);
+	IOProcessBaseHelper::AddReadInputPassForKind(GraphBuilder, ResourceAccess, EResourceName::Flow, ResourceMapping, BufferUAV);
+	IOProcessBaseHelper::AddReadInputPassForKind(GraphBuilder, ResourceAccess, EResourceName::Output, ResourceMapping, BufferUAV);
+}
+
+bool FOutputProcessBase::Validate(const IModelInstance& ModelInstance, FIntPoint Extent) const
+{
+	SCOPED_NAMED_EVENT_TEXT("NNEDenoiser.Prepare", FColor::Magenta);
+
+	const int32 NumBatches = 1;
+
+	TConstArrayView<UE::NNE::FTensorDesc> OutputTensorDescs = ModelInstance.GetOutputTensorDescs();
+	if (OutputTensorDescs.Num() != OutputLayout.Num())
+	{
+		UE_LOG(LogNNEDenoiser, Error, TEXT("Wrong number of outputs (expected %d, got %d)!"), OutputLayout.Num(), OutputTensorDescs.Num())
+		return false;
+	}
+
+	for (int32 Idx = 0; Idx < OutputTensorDescs.Num(); Idx++)
+	{
+		TConstArrayView<int32> OutputSymbolicShapeData = OutputTensorDescs[Idx].GetShape().GetData();
+		const TArray<int32, TInlineAllocator<4>> RequiredOutputShapeData = { NumBatches, OutputLayout.NumChannels(Idx), -1, -1 };
+
+		if (!IsTensorShapeValid(OutputSymbolicShapeData, RequiredOutputShapeData, TEXT("Output")))
+		{
+			return false;
+		}
+	}
+
+	TConstArrayView<UE::NNE::FTensorShape> OutputShapes = ModelInstance.GetOutputTensorShapes();
+	if (!OutputShapes.IsEmpty() && OutputShapes.Num() != OutputLayout.Num())
+	{
+		UE_LOG(LogNNEDenoiser, Error, TEXT("Wrong number of output shapes!"));
+		return false;
+	}
+
+	if (OutputShapes.IsEmpty() && OutputLayout.Num() > 0)
+	{
+		UE_LOG(LogNNEDenoiser, Log, TEXT("Output shapes not resolved yet"));
+		return true;
+	}
+	
+	UE_LOG(LogNNEDenoiser, Log, TEXT("Output shapes (resolved):"));
+
+	for (int32 Idx = 0; Idx < OutputTensorDescs.Num(); Idx++)
+	{
+		TConstArrayView<uint32> OutputShapeData = OutputShapes[Idx].GetData();
+		const TArray<int32, TInlineAllocator<4>> RequiredOutputShapeData = { NumBatches, OutputLayout.NumChannels(Idx), -1, -1 };
+
+		if (!IsTensorShapeValid(OutputShapeData, RequiredOutputShapeData, TEXT("Output")))
+		{
+			return false;
+		}
+	
+		UE_LOG(LogNNEDenoiser, Log, TEXT("%d: (%d, %d, %d, %d)"), Idx, OutputShapeData[0], OutputShapeData[1], OutputShapeData[2], OutputShapeData[3]);
+	}
+
+	return true;
+}
+
+void FOutputProcessBase::AddPasses(
+	FRDGBuilder& GraphBuilder,
+	TConstArrayView<UE::NNE::FTensorDesc> TensorDescs,
+	TConstArrayView<UE::NNE::FTensorShape> TensorShapes,
+	const IResourceAccess& ResourceAccess,
+	TConstArrayView<FRDGBufferRef> Buffers,
+	FRDGTextureRef OutputTexture) const
+{
+	for (int32 I = 0; I < TensorDescs.Num(); I++)
+	{
+		// outputs might be discarded
+		const FResourceMapping* ResourceMapping = OutputLayout.Get(I);
+		if (ResourceMapping && ResourceMapping->GetChannelMappingPerFrame(EResourceName::Output).Num() == 1)
+		{
+			ReadOutputBuffer(GraphBuilder, TensorDescs[I], TensorShapes[I], ResourceAccess, Buffers[I], *ResourceMapping, OutputTexture);
+		}
+	}
+
+	const EResourceName TensorName = EResourceName::Output;
+	const int32 FrameIdx = 0;
+
+	if (HasPostprocessOutput(TensorName, FrameIdx))
+	{
+		FRDGTextureRef PostprocessInputTexture = ResourceAccess.GetTexture(EResourceName::Output, FrameIdx);
+		FRDGTextureRef PostprocessOutputTexture = ResourceAccess.GetIntermediateTexture(EResourceName::Output, FrameIdx);
+		
+		PostprocessOutput(GraphBuilder, PostprocessInputTexture, PostprocessOutputTexture);
+	}
+}
+
+bool FOutputProcessBase::HasPostprocessOutput(EResourceName TensorName, int32 FrameIdx) const
+{
+	check(FrameIdx == 0);
+
+	return TensorName == EResourceName::Output;
+}
+
+void FOutputProcessBase::ReadOutputBuffer(
+	FRDGBuilder& GraphBuilder,
+	const UE::NNE::FTensorDesc& TensorDesc,
+	const UE::NNE::FTensorShape& TensorShape,
+	const IResourceAccess& ResourceAccess,
+	FRDGBufferRef Buffer,
+	const FResourceMapping& ResourceMapping,
+	FRDGTextureRef Texture) const
+{
+	const ENNETensorDataType TensorDataType = TensorDesc.GetDataType();
+	const NNEDenoiserShaders::Internal::EDataType DataType = GetDenoiserShaderDataType(TensorDataType);
+	const FIntPoint BufferSize(TensorShape.GetData()[3], TensorShape.GetData()[2]);
+		
+	TMap<int32, TArray<FChannelMapping>> ChannelMappingPerFrame = ResourceMapping.GetChannelMappingPerFrame(EResourceName::Output);
+	checkf(ChannelMappingPerFrame.Num() == 1, TEXT("Invalid output mapping"));
+	const TArray<FChannelMapping>& ChannelMapping = ChannelMappingPerFrame.CreateConstIterator().Value();
+
+	const EPixelFormat BufferFormat = GetBufferFormat(TensorDataType);
+
+	IOProcessBaseHelper::AddWriteOutputPass(GraphBuilder, Buffer, BufferFormat, Texture, DataType, ChannelMapping);
+}
+
+void FOutputProcessBase::PostprocessOutput(
+		FRDGBuilder& GraphBuilder,
+		FRDGTextureRef Texture,
+		FRDGTextureRef PostprocessedTexture) const
+{
+	if (TransferFunction.IsValid())
+	{
+		TransferFunction->RDGInverse(GraphBuilder, Texture, PostprocessedTexture);
+	}
+
+	IOProcessBaseHelper::AddPreOrPostProcess(GraphBuilder, PostprocessedTexture, EResourceName::Output, 0, PostprocessedTexture);
+}
+
+} // namespace UE::NNEDenoiser::Private
\ No newline at end of file
diff --git a/Engine/Plugins/NNE/NNEDenoiser/Source/NNEDenoiser/Private/NNEDenoiserIOProcessBase.h b/Engine/Plugins/NNE/NNEDenoiser/Source/NNEDenoiser/Private/NNEDenoiserIOProcessBase.h
new file mode 100644
--- /dev/null
+++ b/Engine/Plugins/NNE/NNEDenoiser/Source/NNEDenoiser/Private/NNEDenoiserIOProcessBase.h
@@ -0,0 +1,101 @@
+// Copyright Epic Games, Inc. All Rights Reserved.
+
+#pragma once
+
+#include "NNEDenoiserIOProcess.h"
+#include "NNEDenoiserResourceMapping.h"
+
+namespace UE::NNEDenoiser::Private
+{
+
+class ITransferFunction;
+
+class FInputProcessBase : public IInputProcess
+{
+public:
+	FInputProcessBase(FResourceMappingList InputLayout, TSharedPtr<ITransferFunction> TransferFunction) : InputLayout(InputLayout), TransferFunction(TransferFunction)
+	{
+
+	}
+
+	virtual ~FInputProcessBase() = default;
+
+	virtual bool Validate(const IModelInstance& ModelInstance, FIntPoint Extent) const override;
+
+	virtual bool Prepare(IModelInstance& ModelInstance, FIntPoint Extent) const override;
+
+	virtual int32 NumFrames(EResourceName Name) const override;
+
+	void AddPasses(
+		FRDGBuilder& GraphBuilder,
+		TConstArrayView<UE::NNE::FTensorDesc> TensorDescs,
+		TConstArrayView<UE::NNE::FTensorShape> TensorShapes,
+		const IResourceAccess& ResourceAccess,
+		TConstArrayView<FRDGBufferRef> OutputBuffers) const override;
+
+protected:
+	virtual bool HasPreprocessInput(EResourceName TensorName, int32 FrameIdx) const;
+
+	virtual void PreprocessInput(
+		FRDGBuilder& GraphBuilder,
+		FRDGTextureRef Texture,
+		EResourceName TensorName,
+		int32 FrameIdx,
+		FRDGTextureRef PreprocessedTexture) const;
+
+	virtual void WriteInputBuffer(
+		FRDGBuilder& GraphBuilder,
+		const UE::NNE::FTensorDesc& TensorDesc,
+		const UE::NNE::FTensorShape& TensorShape,
+		const IResourceAccess& ResourceAccess,
+		const FResourceMapping& ResourceMapping,
+		FRDGBufferRef Buffer) const;
+
+private:
+	FResourceMappingList InputLayout;
+	TSharedPtr<ITransferFunction> TransferFunction;
+};
+
+class FOutputProcessBase : public IOutputProcess
+{
+public:
+	FOutputProcessBase(FResourceMappingList OutputLayout, TSharedPtr<ITransferFunction> TransferFunction) : OutputLayout(OutputLayout), TransferFunction(TransferFunction)
+	{
+
+	}
+
+	virtual ~FOutputProcessBase() = default;
+
+	virtual bool Validate(const IModelInstance& ModelInstance, FIntPoint Extent) const override;
+
+	void AddPasses(
+		FRDGBuilder& GraphBuilder,
+		TConstArrayView<UE::NNE::FTensorDesc> TensorDescs,
+		TConstArrayView<UE::NNE::FTensorShape> TensorShapes,
+		const IResourceAccess& ResourceAccess,
+		TConstArrayView<FRDGBufferRef> Buffers,
+		FRDGTextureRef OutputTexture) const override;
+
+protected:
+	virtual bool HasPostprocessOutput(EResourceName TensorName, int32 FrameIdx) const;
+
+	virtual void ReadOutputBuffer(
+		FRDGBuilder& GraphBuilder,
+		const UE::NNE::FTensorDesc& TensorDesc,
+		const UE::NNE::FTensorShape& TensorShape,
+		const IResourceAccess& ResourceAccess,
+		FRDGBufferRef Buffer,
+		const FResourceMapping& ResourceMapping,
+		FRDGTextureRef Texture) const;
+
+	virtual void PostprocessOutput(
+		FRDGBuilder& GraphBuilder,
+		FRDGTextureRef Texture,
+		FRDGTextureRef PostprocessedTexture) const;
+
+private:
+	FResourceMappingList OutputLayout;
+	TSharedPtr<ITransferFunction> TransferFunction;
+};
+
+} // namespace UE::NNEDenoiser::Private
\ No newline at end of file
diff --git a/Engine/Plugins/NNE/NNEDenoiser/Source/NNEDenoiser/Private/NNEDenoiserLog.h b/Engine/Plugins/NNE/NNEDenoiser/Source/NNEDenoiser/Private/NNEDenoiserLog.h
new file mode 100644
--- /dev/null
+++ b/Engine/Plugins/NNE/NNEDenoiser/Source/NNEDenoiser/Private/NNEDenoiserLog.h
@@ -0,0 +1,8 @@
+// Copyright Epic Games, Inc. All Rights Reserved.
+
+#pragma once
+
+#include "CoreMinimal.h"
+#include "Logging/LogMacros.h"
+
+DECLARE_LOG_CATEGORY_EXTERN(LogNNEDenoiser, Log, All);
\ No newline at end of file
diff --git a/Engine/Plugins/NNE/NNEDenoiser/Source/NNEDenoiser/Private/NNEDenoiserModelInstance.h b/Engine/Plugins/NNE/NNEDenoiser/Source/NNEDenoiser/Private/NNEDenoiserModelInstance.h
new file mode 100644
--- /dev/null
+++ b/Engine/Plugins/NNE/NNEDenoiser/Source/NNEDenoiser/Private/NNEDenoiserModelInstance.h
@@ -0,0 +1,19 @@
+// Copyright Epic Games, Inc. All Rights Reserved.
+
+#pragma once
+
+#include "NNERuntimeRDG.h"
+
+namespace UE::NNEDenoiser::Private
+{
+
+class IModelInstance : public NNE::IModelInstanceRDG
+{
+public:
+	using ESetInputTensorShapesStatus = NNE::IModelInstanceRDG::ESetInputTensorShapesStatus;
+	using EEnqueueRDGStatus = NNE::IModelInstanceRDG::EEnqueueRDGStatus;
+
+	virtual ~IModelInstance() = default;
+};
+
+}
\ No newline at end of file
diff --git a/Engine/Plugins/NNE/NNEDenoiser/Source/NNEDenoiser/Private/NNEDenoiserModelInstanceCPU.cpp b/Engine/Plugins/NNE/NNEDenoiser/Source/NNEDenoiser/Private/NNEDenoiserModelInstanceCPU.cpp
new file mode 100644
--- /dev/null
+++ b/Engine/Plugins/NNE/NNEDenoiser/Source/NNEDenoiser/Private/NNEDenoiserModelInstanceCPU.cpp
@@ -0,0 +1,165 @@
+// Copyright Epic Games, Inc. All Rights Reserved.
+
+#include "NNEDenoiserModelInstanceCPU.h"
+#include "NNE.h"
+#include "NNEDenoiserLog.h"
+#include "NNEDenoiserUtils.h"
+#include "NNEModelData.h"
+#include "NNERuntimeCPU.h"
+#include "RenderGraphBuilder.h"
+#include "RHICommandList.h"
+#include "RHIResources.h"
+#include "RHITypes.h"
+#include "UObject/UObjectGlobals.h"
+
+BEGIN_SHADER_PARAMETER_STRUCT(FNNEDenoiserModelInstanceCPUTextureParameters, )
+	RDG_BUFFER_ACCESS_ARRAY(InputBuffers)
+	RDG_BUFFER_ACCESS_ARRAY(OutputBuffers)
+END_SHADER_PARAMETER_STRUCT()
+
+namespace UE::NNEDenoiser::Private
+{
+
+	TUniquePtr<FModelInstanceCPU> FModelInstanceCPU::Make(UNNEModelData& ModelData, const FString& RuntimeName)
+	{
+		check(!RuntimeName.IsEmpty());
+
+		TWeakInterfacePtr<INNERuntimeCPU> Runtime = UE::NNE::GetRuntime<INNERuntimeCPU>(RuntimeName);
+		if (!Runtime.IsValid())
+		{
+			UE_LOG(LogNNEDenoiser, Log, TEXT("Could not create model instance. No CPU runtime '%s' found. Valid CPU runtimes are: "), *RuntimeName);
+			for (const FString& Name : UE::NNE::GetAllRuntimeNames<INNERuntimeCPU>())
+			{
+				UE_LOG(LogNNEDenoiser, Log, TEXT("- %s"), *Name);
+			}
+			return {};
+		}
+
+		if (Runtime->CanCreateModelCPU(ModelData) != INNERuntimeCPU::ECanCreateModelCPUStatus::Ok)
+		{
+			UE_LOG(LogNNEDenoiser, Log, TEXT("%s on CPU can not create model"), *RuntimeName);
+			return {};
+		}
+
+		TSharedPtr<UE::NNE::IModelCPU> Model = Runtime->CreateModelCPU(ModelData);
+		if (!Model.IsValid())
+		{
+			UE_LOG(LogNNEDenoiser, Log, TEXT("Could not create model using %s on CPU"), *RuntimeName);
+			return {};
+		}
+
+		TSharedPtr<UE::NNE::IModelInstanceCPU> ModelInstance = Model->CreateModelInstanceCPU();
+		if (!ModelInstance.IsValid())
+		{
+			UE_LOG(LogNNEDenoiser, Log, TEXT("Could not create model instance using %s on CPU"), *RuntimeName);
+			return {};
+		}
+
+		return MakeUnique<FModelInstanceCPU>(ModelInstance.ToSharedRef());
+	}
+
+	FModelInstanceCPU::FModelInstanceCPU(TSharedRef<UE::NNE::IModelInstanceCPU> ModelInstance) :
+			ModelInstance(ModelInstance)
+	{
+
+	}
+
+	FModelInstanceCPU::~FModelInstanceCPU()
+	{
+		
+	}
+
+	TConstArrayView<NNE::FTensorDesc> FModelInstanceCPU::GetInputTensorDescs() const
+	{
+		return ModelInstance->GetInputTensorDescs();
+	}
+
+	TConstArrayView<NNE::FTensorDesc> FModelInstanceCPU::GetOutputTensorDescs() const
+	{
+		return ModelInstance->GetOutputTensorDescs();
+	}
+
+	TConstArrayView<NNE::FTensorShape> FModelInstanceCPU::GetInputTensorShapes() const
+	{
+		return ModelInstance->GetInputTensorShapes();
+	}
+
+	TConstArrayView<NNE::FTensorShape> FModelInstanceCPU::GetOutputTensorShapes() const
+	{
+		return ModelInstance->GetOutputTensorShapes();
+	}
+
+	FModelInstanceCPU::ESetInputTensorShapesStatus FModelInstanceCPU::SetInputTensorShapes(TConstArrayView<NNE::FTensorShape> InInputShapes)
+	{
+		NNE::IModelInstanceCPU::ESetInputTensorShapesStatus Status = ModelInstance->SetInputTensorShapes(InInputShapes);
+
+		return Status == NNE::IModelInstanceCPU::ESetInputTensorShapesStatus::Ok ? ESetInputTensorShapesStatus::Ok : ESetInputTensorShapesStatus::Fail;
+	}
+
+	FModelInstanceCPU::EEnqueueRDGStatus FModelInstanceCPU::EnqueueRDG(FRDGBuilder &GraphBuilder, TConstArrayView<NNE::FTensorBindingRDG> Inputs, TConstArrayView<NNE::FTensorBindingRDG> Outputs)
+	{
+		FNNEDenoiserModelInstanceCPUTextureParameters* DenoiserParameters = GraphBuilder.AllocParameters<FNNEDenoiserModelInstanceCPUTextureParameters>();
+		for (const NNE::FTensorBindingRDG& Binding : Inputs)
+		{
+			DenoiserParameters->InputBuffers.Emplace(Binding.Buffer, ERHIAccess::CopySrc);
+		}
+		for (const NNE::FTensorBindingRDG& Binding : Outputs)
+		{
+			DenoiserParameters->OutputBuffers.Emplace(Binding.Buffer, ERHIAccess::CopyDest);
+		}
+		
+		GraphBuilder.AddPass(RDG_EVENT_NAME("NNEDenoiser.DenoiseCPU"), DenoiserParameters, ERDGPassFlags::Readback,
+		[this, DenoiserParameters](FRHICommandListImmediate& RHICmdList)
+		{
+			SCOPED_NAMED_EVENT_TEXT("FModelInstanceCPU::DenoisePass", FColor::Magenta);
+			
+#if WITH_EDITOR
+			// NOTE: the time will include the transfer from GPU to CPU which will include waiting for the GPU pipeline to complete
+			uint64 FilterExecuteTime = 0;
+			FilterExecuteTime -= FPlatformTime::Cycles64();
+#endif
+
+			ScratchInputBuffers.SetNum(DenoiserParameters->InputBuffers.Num());
+			TArray<NNE::FTensorBindingCPU> InputBindings;
+			for (int32 Idx = 0; Idx < DenoiserParameters->InputBuffers.Num(); Idx++)
+			{
+				ScratchInputBuffers[Idx].SetNumUninitialized(DenoiserParameters->InputBuffers[Idx].GetBuffer()->GetRHI()->GetSize());
+
+				InputBindings.Emplace(NNE::FTensorBindingCPU{(void *)ScratchInputBuffers[Idx].GetData(), (uint64)ScratchInputBuffers[Idx].Num()});
+			}
+
+			ScratchOutputBuffers.SetNum(DenoiserParameters->OutputBuffers.Num());
+			TArray<NNE::FTensorBindingCPU> OutputBindings;
+			for (int32 Idx = 0; Idx < DenoiserParameters->OutputBuffers.Num(); Idx++)
+			{
+				ScratchOutputBuffers[Idx].SetNumUninitialized(DenoiserParameters->OutputBuffers[Idx].GetBuffer()->GetRHI()->GetSize());
+
+				OutputBindings.Emplace(NNE::FTensorBindingCPU{(void *)ScratchOutputBuffers[Idx].GetData(), (uint64)ScratchOutputBuffers[Idx].Num()});
+			}
+
+			for (int32 Idx = 0; Idx < DenoiserParameters->InputBuffers.Num(); Idx++)
+			{
+				FRHIBuffer* Buffer = DenoiserParameters->InputBuffers[Idx].GetBuffer()->GetRHI();
+				CopyBufferFromGPUToCPU(RHICmdList, Buffer, Buffer->GetSize(), ScratchInputBuffers[Idx]);
+			}
+
+			NNE::IModelInstanceCPU::ERunSyncStatus Status = ModelInstance->RunSync(InputBindings, OutputBindings);
+			checkf(Status == NNE::IModelInstanceCPU::ERunSyncStatus::Ok, TEXT("RunSync failed with status %d"), static_cast<int>(Status))
+
+			for (int32 Idx = 0; Idx < DenoiserParameters->OutputBuffers.Num(); Idx++)
+			{
+				FRHIBuffer* Buffer = DenoiserParameters->OutputBuffers[Idx].GetBuffer()->GetRHI();
+				CopyBufferFromCPUToGPU(RHICmdList, ScratchOutputBuffers[Idx], Buffer->GetSize(), Buffer);
+			}
+
+#if WITH_EDITOR
+		FilterExecuteTime += FPlatformTime::Cycles64();
+		const double FilterExecuteTimeMS = 1000.0 * FPlatformTime::ToSeconds64(FilterExecuteTime);
+		UE_LOG(LogNNEDenoiser, Log, TEXT("Denoised on CPU in %.2f ms"), FilterExecuteTimeMS);
+#endif
+		});
+
+		return EEnqueueRDGStatus::Ok;
+	}
+
+}
\ No newline at end of file
diff --git a/Engine/Plugins/NNE/NNEDenoiser/Source/NNEDenoiser/Private/NNEDenoiserModelInstanceCPU.h b/Engine/Plugins/NNE/NNEDenoiser/Source/NNEDenoiser/Private/NNEDenoiserModelInstanceCPU.h
new file mode 100644
--- /dev/null
+++ b/Engine/Plugins/NNE/NNEDenoiser/Source/NNEDenoiser/Private/NNEDenoiserModelInstanceCPU.h
@@ -0,0 +1,46 @@
+// Copyright Epic Games, Inc. All Rights Reserved.
+
+#pragma once
+
+#include "NNEDenoiserModelInstance.h"
+#include "NNEModelData.h"
+
+class FRDGBuilder;
+
+namespace UE::NNE
+{
+	class IModelInstanceCPU;
+}
+
+namespace UE::NNEDenoiser::Private
+{
+
+class FModelInstanceCPU : public IModelInstance
+{
+public:
+	explicit FModelInstanceCPU(TSharedRef<UE::NNE::IModelInstanceCPU> ModelInstance);
+
+	static TUniquePtr<FModelInstanceCPU> Make(UNNEModelData& ModelData, const FString& RuntimeName);
+
+	virtual ~FModelInstanceCPU();
+
+	virtual TConstArrayView<NNE::FTensorDesc> GetInputTensorDescs() const override;
+
+	virtual TConstArrayView<NNE::FTensorDesc> GetOutputTensorDescs() const override;
+
+	virtual TConstArrayView<NNE::FTensorShape> GetInputTensorShapes() const override;
+
+	virtual TConstArrayView<NNE::FTensorShape> GetOutputTensorShapes() const override;
+
+	virtual ESetInputTensorShapesStatus SetInputTensorShapes(TConstArrayView<NNE::FTensorShape> InInputShapes) override;
+
+	virtual EEnqueueRDGStatus EnqueueRDG(FRDGBuilder &GraphBuilder, TConstArrayView<NNE::FTensorBindingRDG> Inputs, TConstArrayView<NNE::FTensorBindingRDG> Outputs) override;
+
+private:
+	TSharedRef<UE::NNE::IModelInstanceCPU> ModelInstance;
+
+	TArray<TArray<uint8>> ScratchInputBuffers;
+	TArray<TArray<uint8>> ScratchOutputBuffers;
+};
+	
+}
\ No newline at end of file
diff --git a/Engine/Plugins/NNE/NNEDenoiser/Source/NNEDenoiser/Private/NNEDenoiserModelInstanceGPU.cpp b/Engine/Plugins/NNE/NNEDenoiser/Source/NNEDenoiser/Private/NNEDenoiserModelInstanceGPU.cpp
new file mode 100644
--- /dev/null
+++ b/Engine/Plugins/NNE/NNEDenoiser/Source/NNEDenoiser/Private/NNEDenoiserModelInstanceGPU.cpp
@@ -0,0 +1,171 @@
+// Copyright Epic Games, Inc. All Rights Reserved.
+
+#include "NNEDenoiserModelInstanceGPU.h"
+#include "NNE.h"
+#include "NNEDenoiserLog.h"
+#include "NNEDenoiserUtils.h"
+#include "NNEModelData.h"
+#include "NNERuntimeGPU.h"
+#include "PathTracingDenoiser.h"
+#include "RenderGraphBuilder.h"
+#include "RHICommandList.h"
+#include "RHIResources.h"
+#include "RHITypes.h"
+#include "UObject/UObjectGlobals.h"
+
+DECLARE_GPU_STAT_NAMED(FModelInstanceGPU, TEXT("NNEDenoiser.ModelInstanceGPU"));
+
+BEGIN_SHADER_PARAMETER_STRUCT(FNNEDenoiserModelInstanceGPUTextureParameters, )
+	RDG_BUFFER_ACCESS_ARRAY(InputBuffers)
+	RDG_BUFFER_ACCESS_ARRAY(OutputBuffers)
+END_SHADER_PARAMETER_STRUCT()
+
+namespace UE::NNEDenoiser::Private
+{
+
+	TUniquePtr<FModelInstanceGPU> FModelInstanceGPU::Make(UNNEModelData& ModelData, const FString& RuntimeName)
+	{
+		check(!RuntimeName.IsEmpty());
+
+		TWeakInterfacePtr<INNERuntimeGPU> Runtime = UE::NNE::GetRuntime<INNERuntimeGPU>(RuntimeName);
+		if (!Runtime.IsValid())
+		{
+			UE_LOG(LogNNEDenoiser, Log, TEXT("Could not create model instance. No GPU runtime '%s' found. Valid GPU runtimes are: "), *RuntimeName);
+			for (const FString& Name : UE::NNE::GetAllRuntimeNames<INNERuntimeGPU>())
+			{
+				UE_LOG(LogNNEDenoiser, Log, TEXT("- %s"), *Name);
+			}
+			return {};
+		}
+
+		if (Runtime->CanCreateModelGPU(ModelData) != INNERuntimeGPU::ECanCreateModelGPUStatus::Ok)
+		{
+			UE_LOG(LogNNEDenoiser, Log, TEXT("%s on GPU can not create model"), *RuntimeName);
+			return {};
+		}
+
+		TSharedPtr<UE::NNE::IModelGPU> Model = Runtime->CreateModelGPU(ModelData);
+		if (!Model.IsValid())
+		{
+			UE_LOG(LogNNEDenoiser, Log, TEXT("Could not create model using %s on GPU"), *RuntimeName);
+			return {};
+		}
+
+		TSharedPtr<UE::NNE::IModelInstanceGPU> ModelInstance = Model->CreateModelInstanceGPU();
+		if (!ModelInstance.IsValid())
+		{
+			UE_LOG(LogNNEDenoiser, Log, TEXT("Could not create model instance using %s on GPU"), *RuntimeName);
+			return {};
+		}
+
+		return MakeUnique<FModelInstanceGPU>(ModelInstance.ToSharedRef());
+	}
+
+	FModelInstanceGPU::FModelInstanceGPU(TSharedRef<UE::NNE::IModelInstanceGPU> ModelInstance) :
+		ModelInstance(ModelInstance)
+	{
+
+	}
+
+	FModelInstanceGPU::~FModelInstanceGPU()
+	{
+		
+	}
+
+	TConstArrayView<NNE::FTensorDesc> FModelInstanceGPU::GetInputTensorDescs() const
+	{
+		return ModelInstance->GetInputTensorDescs();
+	}
+
+	TConstArrayView<NNE::FTensorDesc> FModelInstanceGPU::GetOutputTensorDescs() const
+	{
+		return ModelInstance->GetOutputTensorDescs();
+	}
+
+	TConstArrayView<NNE::FTensorShape> FModelInstanceGPU::GetInputTensorShapes() const
+	{
+		return ModelInstance->GetInputTensorShapes();
+	}
+
+	TConstArrayView<NNE::FTensorShape> FModelInstanceGPU::GetOutputTensorShapes() const
+	{
+		return ModelInstance->GetOutputTensorShapes();
+	}
+
+	FModelInstanceGPU::ESetInputTensorShapesStatus FModelInstanceGPU::SetInputTensorShapes(TConstArrayView<NNE::FTensorShape> InInputShapes)
+	{
+		NNE::IModelInstanceGPU::ESetInputTensorShapesStatus Status = ModelInstance->SetInputTensorShapes(InInputShapes);
+
+		return Status == NNE::IModelInstanceGPU::ESetInputTensorShapesStatus::Ok ? ESetInputTensorShapesStatus::Ok : ESetInputTensorShapesStatus::Fail;
+	}
+
+	FModelInstanceGPU::EEnqueueRDGStatus FModelInstanceGPU::EnqueueRDG(FRDGBuilder &GraphBuilder, TConstArrayView<NNE::FTensorBindingRDG> Inputs, TConstArrayView<NNE::FTensorBindingRDG> Outputs)
+	{
+		FNNEDenoiserModelInstanceGPUTextureParameters* DenoiserParameters = GraphBuilder.AllocParameters<FNNEDenoiserModelInstanceGPUTextureParameters>();
+		for (const NNE::FTensorBindingRDG& Binding : Inputs)
+		{
+			DenoiserParameters->InputBuffers.Emplace(Binding.Buffer, ERHIAccess::CopySrc);
+		}
+		for (const NNE::FTensorBindingRDG& Binding : Outputs)
+		{
+			DenoiserParameters->OutputBuffers.Emplace(Binding.Buffer, ERHIAccess::CopyDest);
+		}
+
+		RDG_EVENT_SCOPE(GraphBuilder, "NNEDenoiser.ModelInstanceGPU");
+		RDG_GPU_STAT_SCOPE(GraphBuilder, FModelInstanceGPU);
+		
+		GraphBuilder.AddPass(RDG_EVENT_NAME("NNEDenoiser.DenoiseGPU"), DenoiserParameters, ERDGPassFlags::Readback,
+		[this, DenoiserParameters](FRHICommandListImmediate& RHICmdList)
+		{
+			SCOPED_NAMED_EVENT_TEXT("FModelInstanceGPU::DenoisePass", FColor::Magenta);
+			
+#if WITH_EDITOR
+			// NOTE: the time will include the transfer from GPU to CPU which will include waiting for the GPU pipeline to complete
+			uint64 FilterExecuteTime = 0;
+			FilterExecuteTime -= FPlatformTime::Cycles64();
+#endif
+
+			ScratchInputBuffers.SetNum(DenoiserParameters->InputBuffers.Num());
+			TArray<NNE::FTensorBindingCPU> InputBindings;
+			for (int32 Idx = 0; Idx < DenoiserParameters->InputBuffers.Num(); Idx++)
+			{
+				ScratchInputBuffers[Idx].SetNumUninitialized(DenoiserParameters->InputBuffers[Idx].GetBuffer()->GetRHI()->GetSize());
+
+				InputBindings.Emplace(NNE::FTensorBindingCPU{(void *)ScratchInputBuffers[Idx].GetData(), (uint64)ScratchInputBuffers[Idx].Num()});
+			}
+
+			ScratchOutputBuffers.SetNum(DenoiserParameters->OutputBuffers.Num());
+			TArray<NNE::FTensorBindingCPU> OutputBindings;
+			for (int32 Idx = 0; Idx < DenoiserParameters->OutputBuffers.Num(); Idx++)
+			{
+				ScratchOutputBuffers[Idx].SetNumUninitialized(DenoiserParameters->OutputBuffers[Idx].GetBuffer()->GetRHI()->GetSize());
+
+				OutputBindings.Emplace(NNE::FTensorBindingCPU{(void *)ScratchOutputBuffers[Idx].GetData(), (uint64)ScratchOutputBuffers[Idx].Num()});
+			}
+
+			for (int32 Idx = 0; Idx < DenoiserParameters->InputBuffers.Num(); Idx++)
+			{
+				FRHIBuffer* Buffer = DenoiserParameters->InputBuffers[Idx].GetBuffer()->GetRHI();
+				CopyBufferFromGPUToCPU(RHICmdList, Buffer, Buffer->GetSize(), ScratchInputBuffers[Idx]);
+			}
+
+			NNE::IModelInstanceGPU::ERunSyncStatus Status = ModelInstance->RunSync(InputBindings, OutputBindings);
+			checkf(Status == NNE::IModelInstanceGPU::ERunSyncStatus::Ok, TEXT("RunSync failed with status %d"), static_cast<int32>(Status))
+
+			for (int32 Idx = 0; Idx < DenoiserParameters->OutputBuffers.Num(); Idx++)
+			{
+				FRHIBuffer* Buffer = DenoiserParameters->OutputBuffers[Idx].GetBuffer()->GetRHI();
+				CopyBufferFromCPUToGPU(RHICmdList, ScratchOutputBuffers[Idx], Buffer->GetSize(), Buffer);
+			}
+
+#if WITH_EDITOR
+		FilterExecuteTime += FPlatformTime::Cycles64();
+		const double FilterExecuteTimeMS = 1000.0 * FPlatformTime::ToSeconds64(FilterExecuteTime);
+		UE_LOG(LogNNEDenoiser, Log, TEXT("Denoised on GPU in %.2f ms"), FilterExecuteTimeMS);
+#endif
+		});
+
+		return EEnqueueRDGStatus::Ok;
+	}
+
+}
\ No newline at end of file
diff --git a/Engine/Plugins/NNE/NNEDenoiser/Source/NNEDenoiser/Private/NNEDenoiserModelInstanceGPU.h b/Engine/Plugins/NNE/NNEDenoiser/Source/NNEDenoiser/Private/NNEDenoiserModelInstanceGPU.h
new file mode 100644
--- /dev/null
+++ b/Engine/Plugins/NNE/NNEDenoiser/Source/NNEDenoiser/Private/NNEDenoiserModelInstanceGPU.h
@@ -0,0 +1,44 @@
+// Copyright Epic Games, Inc. All Rights Reserved.
+
+#pragma once
+
+#include "NNEDenoiserModelInstance.h"
+#include "NNEModelData.h"
+
+namespace UE::NNE
+{
+	class IModelInstanceGPU;
+}
+
+namespace UE::NNEDenoiser::Private
+{
+
+class FModelInstanceGPU : public IModelInstance
+{
+public:
+	explicit FModelInstanceGPU(TSharedRef<UE::NNE::IModelInstanceGPU> ModelInstance);
+
+	static TUniquePtr<FModelInstanceGPU> Make(UNNEModelData& ModelData, const FString& RuntimeName);
+
+	virtual ~FModelInstanceGPU();
+
+	virtual TConstArrayView<NNE::FTensorDesc> GetInputTensorDescs() const override;
+
+	virtual TConstArrayView<NNE::FTensorDesc> GetOutputTensorDescs() const override;
+
+	virtual TConstArrayView<NNE::FTensorShape> GetInputTensorShapes() const override;
+
+	virtual TConstArrayView<NNE::FTensorShape> GetOutputTensorShapes() const override;
+
+	virtual ESetInputTensorShapesStatus SetInputTensorShapes(TConstArrayView<NNE::FTensorShape> InInputShapes) override;
+
+	virtual EEnqueueRDGStatus EnqueueRDG(FRDGBuilder &GraphBuilder, TConstArrayView<NNE::FTensorBindingRDG> Inputs, TConstArrayView<NNE::FTensorBindingRDG> Outputs) override;
+
+private:
+	TSharedRef<UE::NNE::IModelInstanceGPU> ModelInstance;
+
+	TArray<TArray<uint8>> ScratchInputBuffers;
+	TArray<TArray<uint8>> ScratchOutputBuffers;
+};
+	
+}
\ No newline at end of file
diff --git a/Engine/Plugins/NNE/NNEDenoiser/Source/NNEDenoiser/Private/NNEDenoiserModelInstanceRDG.cpp b/Engine/Plugins/NNE/NNEDenoiser/Source/NNEDenoiser/Private/NNEDenoiserModelInstanceRDG.cpp
new file mode 100644
--- /dev/null
+++ b/Engine/Plugins/NNE/NNEDenoiser/Source/NNEDenoiser/Private/NNEDenoiserModelInstanceRDG.cpp
@@ -0,0 +1,98 @@
+// Copyright Epic Games, Inc. All Rights Reserved.
+
+#include "NNEDenoiserModelInstanceRDG.h"
+#include "NNE.h"
+#include "NNEDenoiserLog.h"
+#include "NNEDenoiserUtils.h"
+#include "NNEModelData.h"
+#include "NNERuntimeRDG.h"
+#include "PathTracingDenoiser.h"
+#include "RenderGraphBuilder.h"
+#include "RHICommandList.h"
+#include "RHIResources.h"
+#include "RHITypes.h"
+#include "UObject/UObjectGlobals.h"
+
+namespace UE::NNEDenoiser::Private
+{
+
+	TUniquePtr<FModelInstanceRDG> FModelInstanceRDG::Make(UNNEModelData& ModelData, const FString& RuntimeName)
+	{
+		check(!RuntimeName.IsEmpty());
+		
+		TWeakInterfacePtr<INNERuntimeRDG> Runtime = UE::NNE::GetRuntime<INNERuntimeRDG>(RuntimeName);
+		if (!Runtime.IsValid())
+		{
+			UE_LOG(LogNNEDenoiser, Log, TEXT("Could not create model instance. No RDG runtime '%s' found. Valid RDG runtimes are: "), *RuntimeName);
+			for (const FString& Name : UE::NNE::GetAllRuntimeNames<INNERuntimeRDG>())
+			{
+				UE_LOG(LogNNEDenoiser, Log, TEXT("- %s"), *Name);
+			}
+			return {};
+		}
+
+		if (Runtime->CanCreateModelRDG(ModelData) != INNERuntimeRDG::ECanCreateModelRDGStatus::Ok)
+		{
+			UE_LOG(LogNNEDenoiser, Log, TEXT("%s on RDG can not create model"), *RuntimeName);
+			return {};
+		}
+
+		TSharedPtr<UE::NNE::IModelRDG> Model = Runtime->CreateModelRDG(ModelData);
+		if (!Model.IsValid())
+		{
+			UE_LOG(LogNNEDenoiser, Log, TEXT("Could not create model using %s on RDG"), *RuntimeName);
+			return {};
+		}
+
+		TSharedPtr<UE::NNE::IModelInstanceRDG> ModelInstance = Model->CreateModelInstanceRDG();
+		if (!ModelInstance.IsValid())
+		{
+			UE_LOG(LogNNEDenoiser, Log, TEXT("Could not create model instance using %s on RDG"), *RuntimeName);
+			return {};
+		}
+
+		return MakeUnique<FModelInstanceRDG>(ModelInstance.ToSharedRef());
+	}
+
+	FModelInstanceRDG::FModelInstanceRDG(TSharedRef<UE::NNE::IModelInstanceRDG> ModelInstance) :
+		ModelInstance(ModelInstance)
+	{
+
+	}
+
+	FModelInstanceRDG::~FModelInstanceRDG()
+	{
+		
+	}
+
+	TConstArrayView<NNE::FTensorDesc> FModelInstanceRDG::GetInputTensorDescs() const
+	{
+		return ModelInstance->GetInputTensorDescs();
+	}
+
+	TConstArrayView<NNE::FTensorDesc> FModelInstanceRDG::GetOutputTensorDescs() const
+	{
+		return ModelInstance->GetOutputTensorDescs();
+	}
+
+	TConstArrayView<NNE::FTensorShape> FModelInstanceRDG::GetInputTensorShapes() const
+	{
+		return ModelInstance->GetInputTensorShapes();
+	}
+
+	TConstArrayView<NNE::FTensorShape> FModelInstanceRDG::GetOutputTensorShapes() const
+	{
+		return ModelInstance->GetOutputTensorShapes();
+	}
+
+	FModelInstanceRDG::ESetInputTensorShapesStatus FModelInstanceRDG::SetInputTensorShapes(TConstArrayView<NNE::FTensorShape> InInputShapes)
+	{
+		return ModelInstance->SetInputTensorShapes(InInputShapes);
+	}
+
+	FModelInstanceRDG::EEnqueueRDGStatus FModelInstanceRDG::EnqueueRDG(FRDGBuilder &GraphBuilder, TConstArrayView<NNE::FTensorBindingRDG> Inputs, TConstArrayView<NNE::FTensorBindingRDG> Outputs)
+	{
+		return ModelInstance->EnqueueRDG(GraphBuilder, Inputs, Outputs);
+	}
+
+}
\ No newline at end of file
diff --git a/Engine/Plugins/NNE/NNEDenoiser/Source/NNEDenoiser/Private/NNEDenoiserModelInstanceRDG.h b/Engine/Plugins/NNE/NNEDenoiser/Source/NNEDenoiser/Private/NNEDenoiserModelInstanceRDG.h
new file mode 100644
--- /dev/null
+++ b/Engine/Plugins/NNE/NNEDenoiser/Source/NNEDenoiser/Private/NNEDenoiserModelInstanceRDG.h
@@ -0,0 +1,41 @@
+// Copyright Epic Games, Inc. All Rights Reserved.
+
+#pragma once
+
+#include "NNEDenoiserModelInstance.h"
+#include "NNEModelData.h"
+
+namespace UE::NNE
+{
+	class IModelInstanceRDG;
+}
+
+namespace UE::NNEDenoiser::Private
+{
+
+class FModelInstanceRDG : public IModelInstance
+{
+public:
+	explicit FModelInstanceRDG(TSharedRef<UE::NNE::IModelInstanceRDG> ModelInstance);
+
+	static TUniquePtr<FModelInstanceRDG> Make(UNNEModelData& ModelData, const FString& RuntimeName);
+
+	virtual ~FModelInstanceRDG();
+
+	virtual TConstArrayView<NNE::FTensorDesc> GetInputTensorDescs() const override;
+
+	virtual TConstArrayView<NNE::FTensorDesc> GetOutputTensorDescs() const override;
+
+	virtual TConstArrayView<NNE::FTensorShape> GetInputTensorShapes() const override;
+
+	virtual TConstArrayView<NNE::FTensorShape> GetOutputTensorShapes() const override;
+
+	virtual ESetInputTensorShapesStatus SetInputTensorShapes(TConstArrayView<NNE::FTensorShape> InInputShapes) override;
+
+	virtual EEnqueueRDGStatus EnqueueRDG(FRDGBuilder &GraphBuilder, TConstArrayView<NNE::FTensorBindingRDG> Inputs, TConstArrayView<NNE::FTensorBindingRDG> Outputs) override;
+
+private:
+	TSharedRef<UE::NNE::IModelInstanceRDG> ModelInstance;
+};
+	
+}
\ No newline at end of file
diff --git a/Engine/Plugins/NNE/NNEDenoiser/Source/NNEDenoiser/Private/NNEDenoiserParameters.h b/Engine/Plugins/NNE/NNEDenoiser/Source/NNEDenoiser/Private/NNEDenoiserParameters.h
new file mode 100644
--- /dev/null
+++ b/Engine/Plugins/NNE/NNEDenoiser/Source/NNEDenoiser/Private/NNEDenoiserParameters.h
@@ -0,0 +1,23 @@
+// Copyright Epic Games, Inc. All Rights Reserved.
+
+#pragma once
+
+#include "HAL/Platform.h"
+
+namespace UE::NNEDenoiser::Private
+{
+	
+	struct FParameters
+	{
+		struct FTilingConfig
+		{
+			int32 Alignment = 1;
+			int32 Overlap{};
+			int32 MaxSize{};
+			int32 MinSize = 1;
+		};
+
+		FTilingConfig TilingConfig{};
+	};
+
+} // namespace UE::NNEDenoiser::Private
\ No newline at end of file
diff --git a/Engine/Plugins/NNE/NNEDenoiser/Source/NNEDenoiser/Private/NNEDenoiserPathTracingDenoiser.cpp b/Engine/Plugins/NNE/NNEDenoiser/Source/NNEDenoiser/Private/NNEDenoiserPathTracingDenoiser.cpp
new file mode 100644
--- /dev/null
+++ b/Engine/Plugins/NNE/NNEDenoiser/Source/NNEDenoiser/Private/NNEDenoiserPathTracingDenoiser.cpp
@@ -0,0 +1,27 @@
+// Copyright Epic Games, Inc. All Rights Reserved.
+
+#include "NNEDenoiserPathTracingDenoiser.h"
+#include "NNEDenoiserGenericDenoiser.h"
+#include "NNEDenoiserHistory.h"
+#include "NNEDenoiserLog.h"
+#include "SceneView.h"
+
+namespace UE::NNEDenoiser::Private
+{
+
+using UE::Renderer::Private::IPathTracingDenoiser;
+
+FPathTracingDenoiser::FPathTracingDenoiser(TUniquePtr<FGenericDenoiser> Denoiser) : Denoiser(MoveTemp(Denoiser))
+{
+
+}
+
+void FPathTracingDenoiser::AddPasses(FRDGBuilder& GraphBuilder, const FSceneView& View, const FInputs& Inputs) const
+{
+	TUniquePtr<FHistory> History = Denoiser->AddPasses(GraphBuilder, Inputs.ColorTex, Inputs.AlbedoTex, Inputs.NormalTex, Inputs.OutputTex, nullptr, View.GPUMask, nullptr);
+	
+	// note: can't store history
+	check(!History.IsValid()); 
+}
+
+} // namespace UE::NNEDenoiser::Private
\ No newline at end of file
diff --git a/Engine/Plugins/NNE/NNEDenoiser/Source/NNEDenoiser/Private/NNEDenoiserPathTracingDenoiser.h b/Engine/Plugins/NNE/NNEDenoiser/Source/NNEDenoiser/Private/NNEDenoiserPathTracingDenoiser.h
new file mode 100644
--- /dev/null
+++ b/Engine/Plugins/NNE/NNEDenoiser/Source/NNEDenoiser/Private/NNEDenoiserPathTracingDenoiser.h
@@ -0,0 +1,31 @@
+// Copyright Epic Games, Inc. All Rights Reserved.
+
+#pragma once
+
+#include "PathTracingDenoiser.h"
+#include "Templates/UniquePtr.h"
+
+namespace UE::NNEDenoiser::Private
+{
+
+class FGenericDenoiser;
+
+using UE::Renderer::Private::IPathTracingDenoiser;
+
+class FPathTracingDenoiser final : public IPathTracingDenoiser
+{
+
+public:
+	explicit FPathTracingDenoiser(TUniquePtr<FGenericDenoiser> Denoiser);
+
+	~FPathTracingDenoiser() = default;
+
+	bool NeedTextureCreateExtraFlags() const override { return true; }
+
+	void AddPasses(FRDGBuilder& GraphBuilder, const FSceneView& View, const FInputs& Inputs) const override;
+
+private:
+	TUniquePtr<FGenericDenoiser> Denoiser;
+};
+
+} // namespace UE::NNEDenoiser::Private
\ No newline at end of file
diff --git a/Engine/Plugins/NNE/NNEDenoiser/Source/NNEDenoiser/Private/NNEDenoiserPathTracingSpatialTemporalDenoiser.cpp b/Engine/Plugins/NNE/NNEDenoiser/Source/NNEDenoiser/Private/NNEDenoiserPathTracingSpatialTemporalDenoiser.cpp
new file mode 100644
--- /dev/null
+++ b/Engine/Plugins/NNE/NNEDenoiser/Source/NNEDenoiser/Private/NNEDenoiserPathTracingSpatialTemporalDenoiser.cpp
@@ -0,0 +1,41 @@
+// Copyright Epic Games, Inc. All Rights Reserved.
+
+#include "NNEDenoiserPathTracingSpatialTemporalDenoiser.h"
+#include "NNEDenoiserGenericDenoiser.h"
+#include "NNEDenoiserHistory.h"
+#include "NNEDenoiserLog.h"
+#include "SceneView.h"
+
+namespace UE::NNEDenoiser::Private
+{
+
+using UE::Renderer::Private::IPathTracingSpatialTemporalDenoiser;
+
+FPathTracingSpatialTemporalDenoiser::FPathTracingSpatialTemporalDenoiser(TUniquePtr<FGenericDenoiser> Denoiser) : Denoiser(MoveTemp(Denoiser))
+{
+
+}
+
+const TCHAR* FPathTracingSpatialTemporalDenoiser::GetDebugName() const
+{
+	return Denoiser->GetDebugName();
+}
+
+IPathTracingSpatialTemporalDenoiser::FOutputs FPathTracingSpatialTemporalDenoiser::AddPasses(FRDGBuilder& GraphBuilder, const FSceneView& View, const FInputs& Inputs) const
+{
+	FHistory* PrevHistory = Inputs.PrevHistory.IsValid() ? reinterpret_cast<FHistory*>(Inputs.PrevHistory.GetReference()) : nullptr;
+	
+	TUniquePtr<FHistory> NewHistory = Denoiser->AddPasses(GraphBuilder, Inputs.ColorTex, Inputs.AlbedoTex, Inputs.NormalTex, Inputs.OutputTex, Inputs.FlowTex, View.GPUMask, PrevHistory);
+
+	FOutputs Outputs;
+	Outputs.NewHistory = NewHistory.Release();
+
+	return Outputs;
+}
+
+void FPathTracingSpatialTemporalDenoiser::AddMotionVectorPass(FRDGBuilder& GraphBuilder, const FSceneView& View, const FMotionVectorInputs& Inputs) const
+{
+	unimplemented();
+}
+
+} // namespace UE::NNEDenoiser::Private
\ No newline at end of file
diff --git a/Engine/Plugins/NNE/NNEDenoiser/Source/NNEDenoiser/Private/NNEDenoiserPathTracingSpatialTemporalDenoiser.h b/Engine/Plugins/NNE/NNEDenoiser/Source/NNEDenoiser/Private/NNEDenoiserPathTracingSpatialTemporalDenoiser.h
new file mode 100644
--- /dev/null
+++ b/Engine/Plugins/NNE/NNEDenoiser/Source/NNEDenoiser/Private/NNEDenoiserPathTracingSpatialTemporalDenoiser.h
@@ -0,0 +1,35 @@
+// Copyright Epic Games, Inc. All Rights Reserved.
+
+#pragma once
+
+#include "PathTracingDenoiser.h"
+#include "Templates/UniquePtr.h"
+
+namespace UE::NNEDenoiser::Private
+{
+
+class FGenericDenoiser;
+
+using UE::Renderer::Private::IPathTracingSpatialTemporalDenoiser;
+
+class FPathTracingSpatialTemporalDenoiser final : public IPathTracingSpatialTemporalDenoiser
+{
+public:
+	explicit FPathTracingSpatialTemporalDenoiser(TUniquePtr<FGenericDenoiser> Denoiser);
+
+	~FPathTracingSpatialTemporalDenoiser() = default;
+
+	// IPathTracingSpatialTemporalDenoiser interface
+	const TCHAR* GetDebugName() const override;
+
+	bool NeedTextureCreateExtraFlags() const override { return true; }
+
+	FOutputs AddPasses(FRDGBuilder& GraphBuilder, const FSceneView& View, const FInputs& Inputs) const override;
+
+	void AddMotionVectorPass(FRDGBuilder& GraphBuilder, const FSceneView& View, const FMotionVectorInputs& Inputs) const override;
+
+private:
+	TUniquePtr<FGenericDenoiser> Denoiser;
+};
+
+} // namespace UE::NNEDenoiser::Private
\ No newline at end of file
diff --git a/Engine/Plugins/NNE/NNEDenoiser/Source/NNEDenoiser/Private/NNEDenoiserResourceManager.cpp b/Engine/Plugins/NNE/NNEDenoiser/Source/NNEDenoiser/Private/NNEDenoiserResourceManager.cpp
new file mode 100644
--- /dev/null
+++ b/Engine/Plugins/NNE/NNEDenoiser/Source/NNEDenoiser/Private/NNEDenoiserResourceManager.cpp
@@ -0,0 +1,221 @@
+// Copyright Epic Games, Inc. All Rights Reserved.
+
+#include "NNEDenoiserResourceManager.h"
+#include "Algo/Transform.h"
+#include "NNEDenoiserUtils.h"
+#include "NNERuntimeRDG.h"
+#include "RenderGraphBuilder.h"
+#include "RenderGraphUtils.h"
+
+namespace UE::NNEDenoiser::Private
+{
+
+void FResourceManager::AddTexture(EResourceName Name, FRDGTextureRef Current, int32 NumFrames)
+{
+	check(!TextureMap.Contains(Name));
+	check(Current != FRDGTextureRef{});
+	check(NumFrames > 0);
+
+	// Init with current texture at index 0
+	TArray<FRDGTextureRef>& Textures = TextureMap.Emplace(Name, TArray<FRDGTextureRef>{ Current });
+
+	// Find render targets in history, register them as external textures to RDG and add starting at index 1
+	if (const auto& RenderTargets = ResourceMap.Find(Name); RenderTargets)
+	{
+		Algo::Transform(*RenderTargets, Textures, [&] (TRefCountPtr<IPooledRenderTarget> RenderTarget)
+		{
+			return GraphBuilder.RegisterExternalTexture(RenderTarget);
+		});
+	}
+
+	auto FillMissingTextures = [&] (TArray<FRDGTextureRef>& Textures, int32 Num)
+	{
+		if (Textures.Num() >= Num)
+		{
+			return;
+		}
+
+		FRDGTextureRef RefTextue = Textures.Last();
+
+		const int32 OldNum = Textures.Num();
+
+		Textures.SetNumUninitialized(Num);
+		for (int32 I = OldNum; I < Textures.Num(); I++)
+		{
+			Textures[I] = RefTextue;
+		}
+	};
+
+	// Repeat "last" valid texture for frames missing their texture
+	FillMissingTextures(Textures, NumFrames);
+
+	// note: Always create intermediate texture for now.
+	const bool bNeedIntermediateTexture = true;
+
+#if UE_BUILD_DEBUG
+	static TSet<FString> DebugNames;
+#endif
+
+	if (bNeedIntermediateTexture)
+	{
+		check(!IntermediateTextureMap.Contains(Name));
+
+		FRDGTextureDesc Desc = FRDGTextureDesc::Create2D(Tiling.TileSize, Current->Desc.Format, Current->Desc.ClearValue, Current->Desc.Flags);
+
+		TArray<FRDGTextureRef>& IntermediateTextures = IntermediateTextureMap.Add(Name);
+		IntermediateTextures.SetNumUninitialized(NumFrames);
+
+		for (int32 I = 0; I < IntermediateTextures.Num(); I++)
+		{
+#if UE_BUILD_DEBUG
+			const FString ResourceName = UEnum::GetValueAsName(Name).ToString();
+			const FString FrameNumberText = FString::FromInt(I);
+			const FString TextureDebugName = FString(TEXT("NNEDenoiser.IntermediateTexture_")) + ResourceName + FrameNumberText;
+
+			const TCHAR* DebugNamePtr = nullptr;
+			if (!DebugNames.Contains(TextureDebugName))
+			{
+				DebugNames.Add(TextureDebugName);
+			}
+
+			DebugNamePtr = **DebugNames.Find(TextureDebugName);
+
+			IntermediateTextures[I] = GraphBuilder.CreateTexture(Desc, DebugNamePtr);
+#else
+			IntermediateTextures[I] = GraphBuilder.CreateTexture(Desc, TEXT("NNEDenoiser.IntermediateTexture"));
+#endif
+		}
+	}
+}
+
+void FResourceManager::BeginTile(int32 TileIndex)
+{
+	check(TileIndex >= 0);
+	check(TileIndex < Tiling.Tiles.Num());
+
+	CurrentTileIdx = TileIndex;
+	const FTile& Tile = Tiling.Tiles[CurrentTileIdx];
+
+	// Copy sub region from original input texture to tile input texture
+	FRHICopyTextureInfo CopyInputInfo{};
+	CopyInputInfo.SourcePosition = {Tile.Position.X + Tile.InputOffsets.Min.X, Tile.Position.Y + Tile.InputOffsets.Min.Y, 0};
+	CopyInputInfo.DestPosition = {Tile.InputOffsets.Min.X, Tile.InputOffsets.Min.Y, 0};
+	CopyInputInfo.Size = {Tiling.TileSize.X + Tile.InputOffsets.Width(), Tiling.TileSize.Y + Tile.InputOffsets.Height(), 1};
+
+	// Only clear if we don't overwrite the whole texture
+	const bool bNeedClear = Tile.InputOffsets.Width() != 0 || Tile.InputOffsets.Height() != 0;
+	const FLinearColor ClearColor = {.0f, 0.f, .0f, .0f};
+	const FLinearColor ClearNormal = {-1.0f, -1.0f, -1.0f, -1.0f};
+
+	for (const auto& KeyValue : TextureMap)
+	{
+		if (!IntermediateTextureMap.Contains(KeyValue.Key))
+		{
+			continue;
+		}
+
+		const TArray<FRDGTextureRef>& Textures = KeyValue.Value;
+		const TArray<FRDGTextureRef>& IntermediateTextures = IntermediateTextureMap.FindChecked(KeyValue.Key);
+		check(Textures.Num() == IntermediateTextures.Num());
+
+		const int32 StartIdx = KeyValue.Key == EResourceName::Output ? 1 : 0;
+		for (int32 I = StartIdx; I < Textures.Num(); I++)
+		{
+			if (bNeedClear)
+			{
+				AddClearRenderTargetPass(GraphBuilder, IntermediateTextures[I], KeyValue.Key == EResourceName::Normal ? ClearNormal : ClearColor);
+			}
+			AddCopyTexturePass(GraphBuilder, Textures[I], IntermediateTextures[I], CopyInputInfo);
+		}
+	}
+}
+
+void FResourceManager::EndTile()
+{
+	if (!IntermediateTextureMap.Contains(EResourceName::Output))
+	{
+		return;
+	}
+
+	check(CurrentTileIdx >= 0);
+	check(CurrentTileIdx < Tiling.Tiles.Num());
+	
+	const FTile& Tile = Tiling.Tiles[CurrentTileIdx];
+
+	FRHICopyTextureInfo CopyOutputInfo{};
+	CopyOutputInfo.SourcePosition = {Tile.OutputOffsets.Min.X, Tile.OutputOffsets.Min.Y, 0};
+	CopyOutputInfo.DestPosition = {Tile.Position.X + Tile.OutputOffsets.Min.X, Tile.Position.Y + Tile.OutputOffsets.Min.Y, 0};
+	CopyOutputInfo.Size = {Tiling.TileSize.X + Tile.OutputOffsets.Width(), Tiling.TileSize.Y + Tile.OutputOffsets.Height(), 1};
+
+	const TArray<FRDGTextureRef>& Textures = TextureMap.FindChecked(EResourceName::Output);
+	const TArray<FRDGTextureRef>& IntermediateTextures = IntermediateTextureMap.FindChecked(EResourceName::Output);
+	check(!Textures.IsEmpty());
+	check(Textures.Num() == IntermediateTextures.Num());
+
+	AddCopyTexturePass(GraphBuilder, IntermediateTextures[0], Textures[0], CopyOutputInfo);
+}
+
+FRDGTextureRef FResourceManager::GetIntermediateTexture(EResourceName Name, int32 FrameIdx) const
+{
+	const TArray<FRDGTextureRef>& IntermediateTextures = IntermediateTextureMap.FindChecked(Name);
+
+	check(FrameIdx >= 0);
+	check(FrameIdx < IntermediateTextures.Num());
+	return IntermediateTextures[FrameIdx];
+}
+
+TMap<EResourceName, TArray<TRefCountPtr<IPooledRenderTarget>>> FResourceManager::MakeHistoryResourceMap()
+{
+	TMap<EResourceName, TArray<TRefCountPtr<IPooledRenderTarget>>> Result;
+
+	for (const auto& KeyValue : TextureMap)
+	{
+		EResourceName Name = KeyValue.Key;
+		const TArray<FRDGTextureRef>& Textures = KeyValue.Value;
+		FRDGTextureRef CurrentTexture = Textures[0];
+
+		const int32 HistoryNumFrames = FMath::Max(Textures.Num() - 1, 0);
+
+		TArray<TRefCountPtr<IPooledRenderTarget>> NewHistoryRenderTargets;
+		if (HistoryNumFrames > 1)
+		{
+			if (const auto* RenderTargets = ResourceMap.Find(Name); RenderTargets)
+			{
+				NewHistoryRenderTargets.Append(RenderTargets->GetData(), FMath::Min(RenderTargets->Num(), HistoryNumFrames - 1));
+			}
+		}
+
+		if (HistoryNumFrames)
+		{
+			GraphBuilder.QueueTextureExtraction(CurrentTexture, &NewHistoryRenderTargets.InsertDefaulted_GetRef(0));
+		}
+
+		if (!NewHistoryRenderTargets.IsEmpty())
+		{
+			Result.Add(Name, MoveTemp(NewHistoryRenderTargets));
+		}
+	}
+
+	return Result;
+}
+
+FRDGTextureRef FResourceManager::GetTexture(EResourceName Name, int32 FrameIdx, bool bBaseTexture) const
+{
+	const bool bIntermediate = !bBaseTexture && IntermediateTextureMap.Contains(Name);
+
+	const TMap<EResourceName, TArray<FRDGTextureRef>>& Map = bIntermediate ? IntermediateTextureMap : TextureMap;
+	check(Map.Contains(Name));
+
+	const TArray<FRDGTextureRef>& Textures = Map.FindChecked(Name);
+	check(FrameIdx >= 0);
+	check(FrameIdx < Textures.Num());
+
+	if (FrameIdx < 0 || FrameIdx >= Textures.Num())
+	{
+		return FRDGTextureRef{};
+	}
+
+	return Textures[FrameIdx];
+}
+
+}
\ No newline at end of file
diff --git a/Engine/Plugins/NNE/NNEDenoiser/Source/NNEDenoiser/Private/NNEDenoiserResourceManager.h b/Engine/Plugins/NNE/NNEDenoiser/Source/NNEDenoiser/Private/NNEDenoiserResourceManager.h
new file mode 100644
--- /dev/null
+++ b/Engine/Plugins/NNE/NNEDenoiser/Source/NNEDenoiser/Private/NNEDenoiserResourceManager.h
@@ -0,0 +1,55 @@
+// Copyright Epic Games, Inc. All Rights Reserved.
+
+#pragma once
+
+#include "Containers/Map.h"
+#include "NNEDenoiserResourceName.h"
+#include "NNEDenoiserTiling.h"
+#include "RendererInterface.h"
+#include "RenderGraphFwd.h"
+
+namespace UE::NNE
+{
+	struct FTensorBindingRDG;
+}
+
+namespace UE::NNEDenoiser::Private
+{
+
+class FResourceManager
+{
+public:
+	FResourceManager(FRDGBuilder& GraphBuilder, FTiling Tiling, TMap<EResourceName, TArray<TRefCountPtr<IPooledRenderTarget>>> ResourceMap = {}) :
+		GraphBuilder(GraphBuilder), Tiling(Tiling), ResourceMap(ResourceMap)
+	{
+
+	}
+
+	void AddTexture(EResourceName Name, FRDGTextureRef Current, int32 NumFrames);
+
+	void BeginTile(int32 TileIndex);
+	void EndTile();
+
+	FRDGTextureRef GetIntermediateTexture(EResourceName Name, int32 FrameIdx) const;
+
+	FRDGTextureRef GetTexture(EResourceName Name, int32 FrameIdx) const
+	{
+		return GetTexture(Name, FrameIdx, false);
+	}
+
+	TMap<EResourceName, TArray<TRefCountPtr<IPooledRenderTarget>>> MakeHistoryResourceMap();
+
+private:
+	FRDGTextureRef GetTexture(EResourceName Name, int32 FrameIdx, bool bBaseTexture) const;
+	
+	FRDGBuilder& GraphBuilder;
+	FTiling Tiling;
+	TMap<EResourceName, TArray<TRefCountPtr<IPooledRenderTarget>>> ResourceMap;
+
+	TMap<EResourceName, TArray<FRDGTextureRef>> TextureMap;
+	TMap<EResourceName, TArray<FRDGTextureRef>> IntermediateTextureMap;
+
+	int32 CurrentTileIdx = -1;
+};
+
+} // namespace UE::NNEDenoiser::Private
\ No newline at end of file
diff --git a/Engine/Plugins/NNE/NNEDenoiser/Source/NNEDenoiser/Private/NNEDenoiserResourceMapping.cpp b/Engine/Plugins/NNE/NNEDenoiser/Source/NNEDenoiser/Private/NNEDenoiserResourceMapping.cpp
new file mode 100644
--- /dev/null
+++ b/Engine/Plugins/NNE/NNEDenoiser/Source/NNEDenoiser/Private/NNEDenoiserResourceMapping.cpp
@@ -0,0 +1,244 @@
+// Copyright Epic Games, Inc. All Rights Reserved.
+
+#include "NNEDenoiserResourceMapping.h"
+
+namespace UE::NNEDenoiser::Private
+{
+
+namespace Detail
+{
+
+FResourceMappingList MakeTensorLayout(const TMap<int32, TMap<int32, FResourceInfo>>& ResourceMap)
+{
+	FResourceMappingList Result{};
+
+	for (int32 I = 0; I < ResourceMap.Num(); I++)
+	{
+		checkf(ResourceMap.Contains(I), TEXT("Missing intput/output %d, must be continuous!"), I);
+
+		const TMap<int32, FResourceInfo>& InnerMap = ResourceMap[I];
+		FResourceMapping& Mapping = Result.Add_GetRef({});
+
+		for (int32 J = 0; J < InnerMap.Num(); J++)
+		{
+			checkf(InnerMap.Contains(J), TEXT("Missing tensor info for channel %d, must be continuous!"), J);
+
+			Mapping.Add(InnerMap[J]);
+		}
+	}
+
+	return Result;
+}
+
+
+void Add(TMap<int32, TMap<int32, FResourceInfo>>& Map, int32 TensorIndex, int32 TensorChannel, EResourceName ResourceName, int32 ResourceChannel, int32 FrameIndex)
+{
+	if (TensorChannel < 0)
+	{
+		for (int32 I = 0; I < -TensorChannel; I++)
+		{
+			Map.FindOrAdd(TensorIndex).FindOrAdd(I) = FResourceInfo{ResourceName, I, 0};
+		}
+	}
+	else
+	{
+		Map.FindOrAdd(TensorIndex).FindOrAdd(TensorChannel) = FResourceInfo{ResourceName, ResourceChannel, 0};
+	}
+}
+
+} // namespace Detail
+
+void FResourceMapping::Add(FResourceInfo Info)
+{
+	ChannelMapping.Add(MoveTemp(Info));
+}
+
+FResourceInfo& FResourceMapping::Add_GetRef(FResourceInfo Info)
+{
+	return ChannelMapping.Add_GetRef(MoveTemp(Info));
+}
+
+const FResourceInfo& FResourceMapping::GetChecked(int32 Channel) const
+{
+	check(Channel >= 0 && Channel < ChannelMapping.Num());
+
+	return ChannelMapping[Channel];
+}
+
+int32 FResourceMapping::Num() const
+{
+	return ChannelMapping.Num();
+}
+
+bool FResourceMapping::HasResource(EResourceName Name) const
+{
+	for (const auto& Info: ChannelMapping)
+	{
+		if (Info.Name == Name)
+		{
+			return true;
+		}
+	}
+	return false;
+}
+
+int32 FResourceMapping::NumFrames(EResourceName Name) const
+{
+	int32 MinFrameIndex = 1;
+	for (const auto& Info: ChannelMapping)
+	{
+		if (Info.Name == Name)
+		{
+			MinFrameIndex = FMath::Min(MinFrameIndex, Info.Frame);
+		}
+	}
+
+	return 1 - MinFrameIndex;
+}
+
+TMap<int32, TArray<FChannelMapping>> FResourceMapping::GetChannelMappingPerFrame(EResourceName Name) const
+{
+	TMap<int32, TArray<FChannelMapping>> Result;
+	for (int32 I = 0; I < ChannelMapping.Num(); I++)
+	{
+		const FResourceInfo& Info = ChannelMapping[I];
+
+		if (Info.Name == Name)
+		{
+			Result.FindOrAdd(Info.Frame).Add({I, Info.Channel});
+		}
+	}
+
+	return Result;
+}
+
+void FResourceMappingList::Add(FResourceMapping ResourceMapping)
+{
+	InputMapping.Add(MoveTemp(ResourceMapping));
+}
+
+FResourceMapping& FResourceMappingList::Add_GetRef(FResourceMapping ResourceMapping)
+{
+	return InputMapping.Add_GetRef(MoveTemp(ResourceMapping));
+}
+
+const FResourceMapping* FResourceMappingList::Get(int32 InputIndex) const
+{
+	if (InputIndex >= 0 && InputIndex < InputMapping.Num())
+	{
+		return &InputMapping[InputIndex];
+	}
+
+	return {};
+}
+
+const FResourceMapping& FResourceMappingList::GetChecked(int32 InputIndex) const
+{
+	check(InputIndex >= 0 && InputIndex < InputMapping.Num());
+
+	return InputMapping[InputIndex];
+}
+
+int32 FResourceMappingList::Num() const
+{
+	return InputMapping.Num();
+}
+
+int32 FResourceMappingList::NumChannels(int32 InputIndex) const
+{
+	check(InputIndex >= 0 && InputIndex < InputMapping.Num());
+	
+	return InputMapping[InputIndex].Num();
+}
+
+int32 FResourceMappingList::NumFrames(EResourceName Name) const
+{
+	int32 MaxNumFrames = 0;
+	for (const auto& Input : InputMapping)
+	{
+		const int32 NumFrames = Input.NumFrames(Name);
+		if (MaxNumFrames < NumFrames)
+		{
+			MaxNumFrames = NumFrames;
+		}
+	}
+
+	return MaxNumFrames;
+}
+
+bool FResourceMappingList::HasResource(EResourceName Name) const
+{
+	for (const auto& Input : InputMapping)
+	{
+		if (Input.HasResource(Name))
+		{
+			return true;
+		}
+	}
+
+	return false;
+}
+
+bool FResourceMappingList::HasResource(int32 InputIndex, EResourceName Name) const
+{
+	check(InputIndex >= 0 && InputIndex < InputMapping.Num());
+
+	return InputMapping[InputIndex].HasResource(Name);
+}
+
+TMap<int32, TArray<FChannelMapping>> FResourceMappingList::GetChannelMappingPerFrame(int32 InputIndex, EResourceName Name) const
+{
+	check(InputIndex >= 0 && InputIndex < InputMapping.Num());
+
+	return InputMapping[InputIndex].GetChannelMappingPerFrame(Name);
+}
+
+template<>
+FResourceMappingList MakeTensorLayout<FNNEDenoiserInputMappingData>(UDataTable* DataTable)
+{
+	TMap<int32, TMap<int32, FResourceInfo>> Map;
+	DataTable->ForeachRow<FNNEDenoiserInputMappingData>("FResourceLayout", [&] (const FName &Key, const FNNEDenoiserInputMappingData &Value)
+	{
+		Detail::Add(Map, Value.TensorIndex, Value.TensorChannel, ToResourceName(Value.Resource), Value.ResourceChannel, 0);
+	});
+
+	return Detail::MakeTensorLayout(Map);
+}
+
+template<>
+FResourceMappingList MakeTensorLayout<FNNEDenoiserOutputMappingData>(UDataTable* DataTable)
+{
+	TMap<int32, TMap<int32, FResourceInfo>> Map;
+	DataTable->ForeachRow<FNNEDenoiserOutputMappingData>("FResourceLayout", [&] (const FName &Key, const FNNEDenoiserOutputMappingData &Value)
+	{
+		Detail::Add(Map, Value.TensorIndex, Value.TensorChannel, ToResourceName(Value.Resource), Value.ResourceChannel, 0);
+	});
+
+	return Detail::MakeTensorLayout(Map);
+}
+
+template<>
+FResourceMappingList MakeTensorLayout<FNNEDenoiserTemporalInputMappingData>(UDataTable* DataTable)
+{
+	TMap<int32, TMap<int32, FResourceInfo>> Map;
+	DataTable->ForeachRow<FNNEDenoiserTemporalInputMappingData>("FResourceLayout", [&] (const FName &Key, const FNNEDenoiserTemporalInputMappingData &Value)
+	{
+		Detail::Add(Map, Value.TensorIndex, Value.TensorChannel, ToResourceName(Value.Resource), Value.ResourceChannel, Value.FrameIndex);
+	});
+
+	return Detail::MakeTensorLayout(Map);
+}
+
+template<>
+FResourceMappingList MakeTensorLayout<FNNEDenoiserTemporalOutputMappingData>(UDataTable* DataTable)
+{
+	TMap<int32, TMap<int32, FResourceInfo>> Map;
+	DataTable->ForeachRow<FNNEDenoiserTemporalOutputMappingData>("FResourceLayout", [&] (const FName &Key, const FNNEDenoiserTemporalOutputMappingData &Value)
+	{
+		Detail::Add(Map, Value.TensorIndex, Value.TensorChannel, ToResourceName(Value.Resource), Value.ResourceChannel, 0);
+	});
+
+	return Detail::MakeTensorLayout(Map);
+}
+
+} // namespace UE::NNEDenoiser::Private
\ No newline at end of file
diff --git a/Engine/Plugins/NNE/NNEDenoiser/Source/NNEDenoiser/Private/NNEDenoiserResourceMapping.h b/Engine/Plugins/NNE/NNEDenoiser/Source/NNEDenoiser/Private/NNEDenoiserResourceMapping.h
new file mode 100644
--- /dev/null
+++ b/Engine/Plugins/NNE/NNEDenoiser/Source/NNEDenoiser/Private/NNEDenoiserResourceMapping.h
@@ -0,0 +1,97 @@
+// Copyright Epic Games, Inc. All Rights Reserved.
+
+#pragma once
+
+#include "Containers/Array.h"
+#include "Containers/Map.h"
+#include "NNEDenoiserIOMappingData.h"
+#include "NNEDenoiserResourceName.h"
+
+namespace UE::NNEDenoiser::Private
+{
+
+struct FResourceInfo
+{
+	EResourceName Name;
+	int32 Channel;
+	int32 Frame;
+};
+
+struct FChannelMapping
+{
+	int32 TensorChannel;
+	int32 ResourceChannel;
+};
+
+// Resource mapping for one tensor:
+// 0: ResourceInfo0 -> tensor channel 0 maps to channel ResourceInfo0.Channel of resource ResourceInfo0.Name from frame ResourceInfo0.Frame
+// 1: ...
+class FResourceMapping
+{
+public:
+	// Add mapping from tensor channel Num()-1 to resource defined by 'Info'
+	void Add(FResourceInfo Info);
+
+	// Add mapping from tensor channel Num()-1 to resource defined by 'Info' and return reference to it
+	FResourceInfo& Add_GetRef(FResourceInfo Info);
+
+	const FResourceInfo& GetChecked(int32 Channel) const;
+
+	// Number of mapped tensor channels
+	int32 Num() const;
+
+	// Number of past frames affected by the mapping including current frame for resource 'Name'
+	int32 NumFrames(EResourceName Name) const;
+
+	// Does mapping use resource 'Name'
+	bool HasResource(EResourceName Name) const;
+
+	// Get mapping from tensor channel to resource channel per frame for resource 'Name'
+	TMap<int32, TArray<FChannelMapping>> GetChannelMappingPerFrame(EResourceName Name) const;
+
+private:
+	TArray<FResourceInfo> ChannelMapping;
+};
+
+// Resource mapping list for multiple tensors:
+// 0: ResourceMapping0 -> tensor 0 maps to resources as defined by ResourceMapping0
+// 1: ...
+class FResourceMappingList
+{
+public:
+	// Add resource mapping for tensor Num()-1
+	void Add(FResourceMapping ResourceMapping);
+
+	// Add resource mapping for tensor Num()-1 and return reference to it
+	FResourceMapping& Add_GetRef(FResourceMapping ResourceMapping);
+
+	const FResourceMapping* Get(int32 InputIndex) const;
+	const FResourceMapping& GetChecked(int32 InputIndex) const;
+
+	// Number of mapped tensors
+	int32 Num() const;
+
+	// Number of mapped tensor channels for input tensor 'InputIndex'
+	int32 NumChannels(int32 InputIndex) const;
+
+	// Number of past frames affected by the mapping including current frame for resource 'Name'
+	int32 NumFrames(EResourceName Name) const;
+
+	// Does mapping use resource 'Name'
+	bool HasResource(EResourceName Name) const;
+
+	// Does mapping for tensor 'InputIndex' use resource 'Name'
+	bool HasResource(int32 InputIndex, EResourceName Name) const;
+
+	// Get mapping from tensor channel to resource channel per frame for resource 'Name' and tensor 'InputIndex"
+	TMap<int32, TArray<FChannelMapping>> GetChannelMappingPerFrame(int32 InputIndex, EResourceName Name) const;
+
+private:
+	TArray<FResourceMapping> InputMapping;
+};
+
+// Helper method to make resource mapping list from data table asset using row struct 'RowStructType'
+template<class RowStructType>
+FResourceMappingList MakeTensorLayout(UDataTable* DataTable);
+
+} // namespace UE::NNEDenoiser::Private
\ No newline at end of file
diff --git a/Engine/Plugins/NNE/NNEDenoiser/Source/NNEDenoiser/Private/NNEDenoiserSettings.cpp b/Engine/Plugins/NNE/NNEDenoiser/Source/NNEDenoiser/Private/NNEDenoiserSettings.cpp
new file mode 100644
--- /dev/null
+++ b/Engine/Plugins/NNE/NNEDenoiser/Source/NNEDenoiser/Private/NNEDenoiserSettings.cpp
@@ -0,0 +1,27 @@
+// Copyright Epic Games, Inc. All Rights Reserved.
+
+#include "NNEDenoiserSettings.h"
+#include "HAL/IConsoleManager.h"
+#include "Misc/ConfigUtilities.h"
+
+#include UE_INLINE_GENERATED_CPP_BY_NAME(NNEDenoiserSettings)
+
+UNNEDenoiserSettings::UNNEDenoiserSettings()
+{
+	CategoryName = TEXT("Plugins");
+	SectionName = TEXT("NNE Denoiser");
+
+	DenoiserAsset = FSoftObjectPath(TEXT("/NNEDenoiser/NNED_Oidn2-3_Balanced_Alpha.NNED_Oidn2-3_Balanced_Alpha"));
+	TemporalDenoiserAsset = FSoftObjectPath(TEXT("/NNEDenoiser/NNEDT_Oidn2-3_Balanced_Alpha.NNEDT_Oidn2-3_Balanced_Alpha"));
+}
+
+void UNNEDenoiserSettings::PostInitProperties()
+{
+	if (IsTemplate())
+	{
+		// We want the .ini file to have precedence over the CVar constructor, so we apply the ini to the CVar before following the regular UDeveloperSettingsBackedByCVars flow
+		UE::ConfigUtilities::ApplyCVarSettingsFromIni(TEXT("/Script/NNEDenoiser.NNEDenoiserSettings"), *GEngineIni, ECVF_SetByProjectSetting);
+	}
+
+	Super::PostInitProperties();
+}
\ No newline at end of file
diff --git a/Engine/Plugins/NNE/NNEDenoiser/Source/NNEDenoiser/Private/NNEDenoiserTiling.h b/Engine/Plugins/NNE/NNEDenoiser/Source/NNEDenoiser/Private/NNEDenoiserTiling.h
new file mode 100644
--- /dev/null
+++ b/Engine/Plugins/NNE/NNEDenoiser/Source/NNEDenoiser/Private/NNEDenoiserTiling.h
@@ -0,0 +1,125 @@
+// Copyright Epic Games, Inc. All Rights Reserved.
+
+#pragma once
+
+#include "Math/IntPoint.h"
+#include "Math/IntRect.h"
+#include "Math/IntVector.h"
+#include "Math/MathFwd.h"
+#include "NNEDenoiserUtils.h"
+
+namespace UE::NNEDenoiser::Private
+{
+	
+	struct FTile
+	{
+		FTile(FIntPoint Position, FIntRect InputOffsets, FIntRect OutputOffsets) : Position(Position), InputOffsets(InputOffsets), OutputOffsets(OutputOffsets) {}
+		
+		FIntPoint Position;
+		FIntRect InputOffsets;
+		FIntRect OutputOffsets;
+	};
+
+	struct FTiling
+	{
+		FTiling() = default;
+
+		FIntPoint TileSize;
+		FIntPoint Count;
+		TArray<FTile> Tiles;
+	};
+
+	inline FTiling CreateTiling(FIntPoint TargetTileSize, FIntPoint MaxTileSize, FIntPoint MinTileSize, int32 TileAlignment, FIntPoint MinimumOverlap, FIntPoint Size)
+	{
+		auto GetTileSize = [] (int32 Max, int32 Min, int32 Alignment, int32 Overlap, int32 Size)
+		{
+			int NumTiles = 1;
+			int32 Result = FMath::Max(RoundUp(Size, Alignment), Min);
+
+			if (Max <= 0)
+			{
+				return Result;
+			}
+
+			Alignment = FMath::Max(Alignment, 1);
+			Min = FMath::Max(Min, Alignment + Overlap);
+
+			while (Result > Max && Result > Min)
+			{
+				NumTiles++;
+				Result = FMath::Max(RoundUp(CeilDiv(Size - Overlap, NumTiles), Alignment) + Overlap, Min);
+			}
+
+			return Result;
+		};
+
+		auto GetNumTiles = [] (int32 TileSize, int32 Overlap, int32 Size)
+		{
+			return Size > TileSize ? CeilDiv(Size - Overlap, TileSize - Overlap) : 1;
+		};
+
+		auto GetOffsets = [] (int32 Count, int32 TileSize, int32 Overlap, int32 Size)
+		{
+			TArray<int32> Result;
+			for (int32 I = 0; I < Count; I++)
+			{
+				Result.Add(I < Count - 1 ? I * (TileSize - Overlap) : FMath::Max(Size - TileSize, 0));
+			}
+
+			return Result;
+		};
+
+		FTiling Result{};
+		Result.TileSize = {
+			TargetTileSize.X > 0 ? TargetTileSize.X : GetTileSize(MaxTileSize.X, MinTileSize.X, TileAlignment, MinimumOverlap.X, Size.X),
+			TargetTileSize.Y > 0 ? TargetTileSize.Y : GetTileSize(MaxTileSize.Y, MinTileSize.Y, TileAlignment, MinimumOverlap.Y, Size.Y)
+		};
+
+		Result.Count = {
+			GetNumTiles(Result.TileSize.X, MinimumOverlap.X, Size.X),
+			GetNumTiles(Result.TileSize.Y, MinimumOverlap.Y, Size.Y)
+		};
+
+		const FIntPoint TotalOverlap = Result.Count * Result.TileSize - Size;
+		const FIntPoint Overlap = {
+			Result.Count.X == 1 ? 0 : TotalOverlap.X / (Result.Count.X - 1),
+			Result.Count.Y == 1 ? 0 : TotalOverlap.Y / (Result.Count.Y - 1)
+		};
+
+		const FIntPoint HalfOverlap = {
+			FMath::FloorToInt32(Overlap.X / 2.0f),
+			FMath::FloorToInt32(Overlap.Y / 2.0f)
+		};
+
+		const TArray<int32> OffsetsX = GetOffsets(Result.Count.X, Result.TileSize.X, Overlap.X, Size.X);
+		const TArray<int32> OffsetsY = GetOffsets(Result.Count.Y, Result.TileSize.Y, Overlap.Y, Size.Y);
+
+		const int32 InX0 = 0;
+		const int32 InX1 = -FMath::Max(Result.TileSize.X - Size.X, 0);
+		const int32 InY0 = 0;
+		const int32 InY1 = -FMath::Max(Result.TileSize.Y - Size.Y, 0);
+
+		for (int32 Ty = 0; Ty < Result.Count.Y; Ty++)
+		{
+			const int32 Y0 = OffsetsY[Ty];
+			const int32 Y1 = Y0 + Result.TileSize.Y;
+
+			const int32 OutY0 = Result.Count.Y == 1 ? InY0 : Ty > 0 ? HalfOverlap.Y : 0;
+			const int32 OutY1 = Result.Count.Y == 1 ? InY1 : Ty < Result.Count.Y - 1 ? -HalfOverlap.Y : 0;
+
+			for (int32 Tx = 0; Tx < Result.Count.X; Tx++)
+			{
+				const int32 X0 = OffsetsX[Tx];
+				const int32 X1 = X0 + Result.TileSize.X;
+
+				const int32 OutX0 = Result.Count.X == 1 ? InX0 : Tx > 0 ? HalfOverlap.X : 0;
+				const int32 OutX1 = Result.Count.X == 1 ? InX1 : Tx < Result.Count.X - 1 ? -HalfOverlap.X : 0;
+
+				Result.Tiles.Add(FTile{{X0, Y0}, {InX0, InY0, InX1, InY1}, {OutX0, OutY0, OutX1, OutY1}});
+			}
+		}
+
+		return Result;
+	}
+
+}
\ No newline at end of file
diff --git a/Engine/Plugins/NNE/NNEDenoiser/Source/NNEDenoiser/Private/NNEDenoiserTransferFunction.h b/Engine/Plugins/NNE/NNEDenoiser/Source/NNEDenoiser/Private/NNEDenoiserTransferFunction.h
new file mode 100644
--- /dev/null
+++ b/Engine/Plugins/NNE/NNEDenoiser/Source/NNEDenoiser/Private/NNEDenoiserTransferFunction.h
@@ -0,0 +1,27 @@
+// Copyright Epic Games, Inc. All Rights Reserved.
+
+#pragma once
+
+#include "Containers/ContainersFwd.h"
+#include "RenderGraphFwd.h"
+
+struct FLinearColor;
+
+namespace UE::NNEDenoiser::Private
+{
+
+class ITransferFunction
+{
+public:
+	virtual ~ITransferFunction() { }
+
+	virtual void Forward(TConstArrayView<FLinearColor> InputImage, float InputScale, TArray<FLinearColor>& OutputImage) const = 0;
+	virtual void Inverse(TConstArrayView<FLinearColor> InputImage, float InvInputScale, TArray<FLinearColor>& OutputImage) const = 0;
+
+	virtual void RDGSetInputScale(FRDGBufferRef InputScaleBuffer) = 0;
+	
+	virtual void RDGForward(FRDGBuilder& GraphBuilder, FRDGTextureRef InputTexture, FRDGTextureRef OutputTexture) const = 0;
+	virtual void RDGInverse(FRDGBuilder& GraphBuilder, FRDGTextureRef InputTexture, FRDGTextureRef OutputTexture) const = 0;
+};
+
+}
\ No newline at end of file
diff --git a/Engine/Plugins/NNE/NNEDenoiser/Source/NNEDenoiser/Private/NNEDenoiserTransferFunctionOidn.cpp b/Engine/Plugins/NNE/NNEDenoiser/Source/NNEDenoiser/Private/NNEDenoiserTransferFunctionOidn.cpp
new file mode 100644
--- /dev/null
+++ b/Engine/Plugins/NNE/NNEDenoiser/Source/NNEDenoiser/Private/NNEDenoiserTransferFunctionOidn.cpp
@@ -0,0 +1,171 @@
+// Copyright Epic Games, Inc. All Rights Reserved.
+
+#include "NNEDenoiserTransferFunctionOidn.h"
+#include "NNEDenoiserShadersTransferFunctionOidnCS.h"
+#include "Algo/Transform.h"
+
+DECLARE_GPU_STAT_NAMED(FNNEDenoiserTransferFunctionOidn, TEXT("NNEDenoiser.TransferFunctionOidn"));
+
+namespace UE::NNEDenoiser::Private::Oidn
+{
+
+namespace Helper
+{
+
+constexpr float HDRMax = 65504.f; // maximum HDR value
+
+// from [REDACTED_PROFILE_URL]/oidn
+struct FPU
+{
+	static constexpr float a  =  1.41283765e+03f;
+	static constexpr float b  =  1.64593172e+00f;
+	static constexpr float c  =  4.31384981e-01f;
+	static constexpr float d  = -2.94139609e-03f;
+	static constexpr float e  =  1.92653254e-01f;
+	static constexpr float f  =  6.26026094e-03f;
+	static constexpr float g  =  9.98620152e-01f;
+	static constexpr float y0 =  1.57945760e-06f;
+	static constexpr float y1 =  3.22087631e-02f;
+	static constexpr float x0 =  2.23151711e-03f;
+	static constexpr float x1 =  3.70974749e-01f;
+
+	static float Forward(float y)
+	{
+		if (y <= y0)
+		return a * y;
+		else if (y <= y1)
+		return b * FMath::Pow(y, c) + d;
+		else
+		return e * FMath::Loge(y + f) + g;
+	}
+
+	static float Inverse(float x)
+	{
+		if (x <= x0)
+		return x / a;
+		else if (x <= x1)
+		return FMath::Pow((x - d) / b, 1.f/c);
+		else
+		return FMath::Exp((x - g) / e) - f;
+	}
+};
+
+void ApplyTransferFunction(
+	FRDGBuilder& GraphBuilder,
+	FRDGTextureRef InputTexture,
+	FRDGTextureRef OutputTexture,
+	FRDGBufferRef InputScaleBuffer,
+	float NormScale,
+	float InvNormScale,
+	bool Forward)
+{
+	using namespace UE::NNEDenoiserShaders::Internal;
+
+	const FIntVector Size = InputTexture->Desc.GetSize();
+	check(Size == OutputTexture->Desc.GetSize());
+
+	FTransferFunctionOidnCS::FParameters *ShaderParameters = GraphBuilder.AllocParameters<FTransferFunctionOidnCS::FParameters>();
+	ShaderParameters->Width = Size.X;
+	ShaderParameters->Height = Size.Y;
+	ShaderParameters->InputTexture = InputTexture;
+	ShaderParameters->InputScaleBuffer = GraphBuilder.CreateSRV(FRDGBufferSRVDesc(InputScaleBuffer, EPixelFormat::PF_R32_FLOAT));
+	ShaderParameters->NormScale = NormScale;
+	ShaderParameters->InvNormScale = InvNormScale;
+	ShaderParameters->OutputTexture = GraphBuilder.CreateUAV(FRDGTextureUAVDesc(OutputTexture));
+
+	FTransferFunctionOidnCS::FPermutationDomain PermutationVector;
+	PermutationVector.Set<FTransferFunctionOidnCS::FTransferFunctionOidnMode>(Forward ? ETransferFunctionOidnMode::Forwward : ETransferFunctionOidnMode::Inverse);
+
+	FIntVector ThreadGroupCount = FIntVector(
+		FMath::DivideAndRoundUp(Size.X, FTransferFunctionOidnConstants::THREAD_GROUP_SIZE),
+		FMath::DivideAndRoundUp(Size.Y, FTransferFunctionOidnConstants::THREAD_GROUP_SIZE),
+		1);
+
+	FGlobalShaderMap* GlobalShaderMap = GetGlobalShaderMap(GMaxRHIFeatureLevel);
+	TShaderMapRef<FTransferFunctionOidnCS> Shader(GlobalShaderMap, PermutationVector);
+
+	RDG_EVENT_SCOPE(GraphBuilder, "NNEDenoiser.TransferFunctionOidn");
+	RDG_GPU_STAT_SCOPE(GraphBuilder, FNNEDenoiserTransferFunctionOidn);
+
+	FComputeShaderUtils::AddPass(
+		GraphBuilder,
+		RDG_EVENT_NAME("NNEDenoiser.TransferFunctionOidn"),
+		ERDGPassFlags::Compute | ERDGPassFlags::NeverCull,
+		Shader,
+		ShaderParameters,
+		ThreadGroupCount);
+}
+
+} // Helpser
+
+FTransferFunction::FTransferFunction()
+{
+	InvNormScale = Helper::FPU::Forward(Helper::HDRMax);
+	NormScale = 1.0f / InvNormScale;
+}
+
+void FTransferFunction::Forward(TConstArrayView<FLinearColor> InputImage, float InputScale, TArray<FLinearColor>& OutputImage) const
+{
+	auto InputProcess = [this, InputScale] (float Value)
+	{
+		Value *= InputScale;
+		Value = FMath::Clamp(Value, 0.0f, FLT_MAX);
+		
+		return Helper::FPU::Forward(Value) * NormScale;
+	};
+
+	Algo::Transform(InputImage, OutputImage, [&InputProcess] (const FLinearColor& Color)
+	{
+		return FLinearColor(
+			InputProcess(Color.R),
+			InputProcess(Color.G),
+			InputProcess(Color.B),
+			Color.A
+		);
+	});
+}
+
+void FTransferFunction::Inverse(TConstArrayView<FLinearColor> InputImage, float InvInputScale, TArray<FLinearColor>& OutputImage) const
+{
+	auto OutputProcess = [this, InvInputScale] (float Value)
+	{
+		Value = FMath::Clamp(Value, 0.0f, FLT_MAX);
+		Value = Helper::FPU::Inverse(Value * InvNormScale);
+		Value *= InvInputScale;
+
+		return Value;
+	};
+
+	Algo::Transform(InputImage, OutputImage, [&OutputProcess] (const FLinearColor& Color)
+	{
+		return FLinearColor(
+			OutputProcess(Color.R),
+			OutputProcess(Color.G),
+			OutputProcess(Color.B),
+			Color.A
+		);
+	});
+}
+
+void FTransferFunction::RDGSetInputScale(FRDGBufferRef InInputScaleBuffer)
+{
+	check(InInputScaleBuffer);
+
+	InputScaleBuffer = InInputScaleBuffer;
+}
+
+void FTransferFunction::RDGForward(FRDGBuilder& GraphBuilder, FRDGTextureRef InputTexture, FRDGTextureRef OutputTexture) const
+{
+	check(InputScaleBuffer);
+
+	Helper::ApplyTransferFunction(GraphBuilder, InputTexture, OutputTexture, InputScaleBuffer, NormScale, InvNormScale, true);
+}
+
+void FTransferFunction::RDGInverse(FRDGBuilder& GraphBuilder, FRDGTextureRef InputTexture, FRDGTextureRef OutputTexture) const
+{
+	check(InputScaleBuffer);
+	
+	Helper::ApplyTransferFunction(GraphBuilder, InputTexture, OutputTexture, InputScaleBuffer, NormScale, InvNormScale, false);
+}
+
+} // namespace UE::NNEDenoiser::Private::Oidn
\ No newline at end of file
diff --git a/Engine/Plugins/NNE/NNEDenoiser/Source/NNEDenoiser/Private/NNEDenoiserTransferFunctionOidn.h b/Engine/Plugins/NNE/NNEDenoiser/Source/NNEDenoiser/Private/NNEDenoiserTransferFunctionOidn.h
new file mode 100644
--- /dev/null
+++ b/Engine/Plugins/NNE/NNEDenoiser/Source/NNEDenoiser/Private/NNEDenoiserTransferFunctionOidn.h
@@ -0,0 +1,30 @@
+// Copyright Epic Games, Inc. All Rights Reserved.
+
+#pragma once
+
+#include "NNEDenoiserTransferFunction.h"
+
+namespace UE::NNEDenoiser::Private::Oidn
+{
+
+class FTransferFunction : public ITransferFunction
+{
+public:
+	FTransferFunction();
+	virtual ~FTransferFunction() { }
+
+	virtual void Forward(TConstArrayView<FLinearColor> InputImage, float InputScale, TArray<FLinearColor>& OutputImage) const override;
+	virtual void Inverse(TConstArrayView<FLinearColor> InputImage, float InvInputScale, TArray<FLinearColor>& OutputImage) const override;
+
+	virtual void RDGSetInputScale(FRDGBufferRef InInputScaleBuffer) override;
+	
+	virtual void RDGForward(FRDGBuilder& GraphBuilder, FRDGTextureRef InputTexture, FRDGTextureRef OutputTexture) const override;
+	virtual void RDGInverse(FRDGBuilder& GraphBuilder, FRDGTextureRef InputTexture, FRDGTextureRef OutputTexture) const override;
+
+private:
+	float NormScale = 1.0f;
+	float InvNormScale = 1.0f;
+	FRDGBufferRef InputScaleBuffer{};
+};
+
+} // UE::NNEDenoiser::Private::Oidn
\ No newline at end of file
diff --git a/Engine/Plugins/NNE/NNEDenoiser/Source/NNEDenoiser/Private/NNEDenoiserUnitTestOidn.cpp b/Engine/Plugins/NNE/NNEDenoiser/Source/NNEDenoiser/Private/NNEDenoiserUnitTestOidn.cpp
new file mode 100644
--- /dev/null
+++ b/Engine/Plugins/NNE/NNEDenoiser/Source/NNEDenoiser/Private/NNEDenoiserUnitTestOidn.cpp
@@ -0,0 +1,145 @@
+// Copyright Epic Games, Inc. All Rights Reserved.
+
+#include "CoreMinimal.h"
+#include "NNEDenoiserTransferFunctionOidn.h"
+#include "NNEDenoiserUtils.h"
+#include "NNERuntimeCPU.h"
+#include "NNERuntimeRDG.h"
+#include "RenderGraphBuilder.h"
+
+#if WITH_DEV_AUTOMATION_TESTS
+#include "Misc/AutomationTest.h"
+#endif
+
+namespace UE::NNEDenoiser::Private
+{
+
+#if WITH_DEV_AUTOMATION_TESTS
+
+BEGIN_SHADER_PARAMETER_STRUCT(FTransferFunctionTestUploadTextureParameters, )
+	RDG_TEXTURE_ACCESS(InputTexture, ERHIAccess::CopyDest)
+END_SHADER_PARAMETER_STRUCT()
+
+BEGIN_SHADER_PARAMETER_STRUCT(FTransferFunctionTestUploadBufferParameters, )
+	RDG_BUFFER_ACCESS(InputBuffer, ERHIAccess::CopyDest)
+END_SHADER_PARAMETER_STRUCT()
+
+BEGIN_SHADER_PARAMETER_STRUCT(FTransferFunctionTestDownloadParameters, )
+	RDG_TEXTURE_ACCESS(OutputTexture, ERHIAccess::CopySrc)
+END_SHADER_PARAMETER_STRUCT()
+
+IMPLEMENT_SIMPLE_AUTOMATION_TEST(FTransferFunctionTest, "NNEDenoiser.UnitTests.OIDN.TransferFunction.Forward", EAutomationTestFlags_ApplicationContextMask | EAutomationTestFlags_FeatureMask | EAutomationTestFlags::EngineFilter)
+bool FTransferFunctionTest::RunTest(const FString& Parameter)
+{
+	const int32 Width = 1920;
+	const int32 Height = 1088;
+	const int32 Num = Width * Height;
+
+	const TArray<float> InputScale = {0.5f, 2.0f};
+
+	TArray<FLinearColor> TestData;
+	TestData.SetNumUninitialized(Num);
+
+	for (int i = 0; i < TestData.Num(); i++)
+	{
+		TestData[i] = FLinearColor(FMath::FRand(), FMath::FRand(), FMath::FRand(), 1.0);
+	}
+
+	Oidn::FTransferFunction TransferFunction;
+
+	TArray<FLinearColor> ResultCPU;
+	TransferFunction.Forward(TestData, InputScale[0], ResultCPU);
+
+	TArray<FLinearColor> ResultRDG;
+	ResultRDG.SetNumZeroed(Num);
+
+	FEvent* Signal = FGenericPlatformProcess::GetSynchEventFromPool(true);
+
+	ENQUEUE_RENDER_COMMAND(NNETransferFunctionTest)
+	(
+		[&TestData, &ResultRDG, &InputScale, TransferFunction, Num, Signal](FRHICommandListImmediate& RHICmdList) mutable
+		{
+			ERHIPipeline Pipeline = RHICmdList.GetPipeline();
+			if (Pipeline == ERHIPipeline::None)
+			{
+				RHICmdList.SwitchPipeline(ERHIPipeline::Graphics);
+			}
+
+			FRDGBuilder	GraphBuilder(RHICmdList);
+
+			const FRDGTextureDesc InputTextureDesc = FRDGTextureDesc::Create2D({Width, Height}, EPixelFormat::PF_A32B32G32R32F, FClearValueBinding::Black, ETextureCreateFlags::None);
+			FRDGTextureRef InputTexture = GraphBuilder.CreateTexture(InputTextureDesc, TEXT("TransferFunctionInputTexture"));
+
+			{
+				FTransferFunctionTestUploadTextureParameters* Parameters = GraphBuilder.AllocParameters<FTransferFunctionTestUploadTextureParameters>();
+				Parameters->InputTexture = InputTexture;
+
+				GraphBuilder.AddPass(RDG_EVENT_NAME("NNEDenoiser.TransferFunctionTest.UploadInput"), Parameters, ERDGPassFlags::Readback,
+				[&TestData, Parameters](FRHICommandListImmediate& RHICmdList)
+				{
+					SCOPED_NAMED_EVENT_TEXT("FTransferFunctionTest::UploadInput", FColor::Magenta);
+
+					CopyTextureFromCPUToGPU(RHICmdList, TestData, {Width, Height}, Parameters->InputTexture->GetRHI());
+				});
+			}
+
+			FRDGBufferDesc InputScaleBufferDesc = FRDGBufferDesc::CreateBufferDesc(sizeof(float), 2);
+			FRDGBufferRef InputScaleBuffer = GraphBuilder.CreateBuffer(InputScaleBufferDesc, TEXT("TransferFunctionInputScaleBuffer"));
+
+			{
+				FTransferFunctionTestUploadBufferParameters* Parameters = GraphBuilder.AllocParameters<FTransferFunctionTestUploadBufferParameters>();
+				Parameters->InputBuffer = InputScaleBuffer;
+
+				GraphBuilder.AddPass(RDG_EVENT_NAME("NNEDenoiser.TransferFunctionTest.UploadScale"), Parameters, ERDGPassFlags::Readback,
+				[&InputScale, Parameters](FRHICommandListImmediate& RHICmdList)
+				{
+					SCOPED_NAMED_EVENT_TEXT("FTransferFunctionTest::UploadScale", FColor::Magenta);
+
+					CopyBufferFromCPUToGPU(RHICmdList, InputScale, 2, Parameters->InputBuffer->GetRHI());
+				});
+			}
+
+			const FRDGTextureDesc OutputTextureDesc = FRDGTextureDesc::Create2D({Width, Height}, EPixelFormat::PF_A32B32G32R32F, FClearValueBinding::Black, ETextureCreateFlags::UAV);
+			FRDGTextureRef OutputTexture = GraphBuilder.CreateTexture(OutputTextureDesc, TEXT("TransferFunctionOutputTexture"));
+
+			TransferFunction.RDGSetInputScale(InputScaleBuffer);
+
+			TransferFunction.RDGForward(GraphBuilder, InputTexture, OutputTexture);
+
+			{
+				FTransferFunctionTestDownloadParameters* Parameters = GraphBuilder.AllocParameters<FTransferFunctionTestDownloadParameters>();
+				Parameters->OutputTexture = OutputTexture;
+
+				GraphBuilder.AddPass(RDG_EVENT_NAME("NNEDenoiser.TransferFunctionTest.Download"), Parameters, ERDGPassFlags::Readback,
+				[&ResultRDG, Num, Parameters](FRHICommandListImmediate& RHICmdList)
+				{
+					SCOPED_NAMED_EVENT_TEXT("FTransferFunctionTest::Download", FColor::Magenta);
+
+					CopyTextureFromGPUToCPU(RHICmdList, Parameters->OutputTexture->GetRHI(), {Width, Height}, ResultRDG);
+				});
+			}
+
+			GraphBuilder.Execute();
+
+			RHICmdList.BlockUntilGPUIdle();
+
+			Signal->Trigger();
+		}
+	);
+
+	Signal->Wait();
+
+	for (int32 i = 0; i < Num; i++)
+	{
+		UTEST_EQUAL_TOLERANCE(TEXT("TransferFunction"), ResultRDG[i].R, ResultCPU[i].R, 1e-6f);
+		UTEST_EQUAL_TOLERANCE(TEXT("TransferFunction"), ResultRDG[i].G, ResultCPU[i].G, 1e-6f);
+		UTEST_EQUAL_TOLERANCE(TEXT("TransferFunction"), ResultRDG[i].B, ResultCPU[i].B, 1e-6f);
+		UTEST_EQUAL(TEXT("TransferFunction"), ResultRDG[i].A, ResultCPU[i].A);
+	}
+
+	return true;
+}
+
+#endif //WITH_DEV_AUTOMATION_TESTS
+
+} // namespace UE::NNEDenoiser::Internal
diff --git a/Engine/Plugins/NNE/NNEDenoiser/Source/NNEDenoiser/Private/NNEDenoiserUnitTests.cpp b/Engine/Plugins/NNE/NNEDenoiser/Source/NNEDenoiser/Private/NNEDenoiserUnitTests.cpp
new file mode 100644
--- /dev/null
+++ b/Engine/Plugins/NNE/NNEDenoiser/Source/NNEDenoiser/Private/NNEDenoiserUnitTests.cpp
@@ -0,0 +1,114 @@
+// Copyright Epic Games, Inc. All Rights Reserved.
+
+#include "CoreMinimal.h"
+#include "NNEDenoiserAutoExposure.h"
+#include "NNEDenoiserUtils.h"
+#include "NNERuntimeCPU.h"
+#include "NNERuntimeRDG.h"
+#include "RenderGraphBuilder.h"
+
+#if WITH_DEV_AUTOMATION_TESTS
+#include "Misc/AutomationTest.h"
+#endif
+
+namespace UE::NNEDenoiser::Private
+{
+
+#if WITH_DEV_AUTOMATION_TESTS
+
+BEGIN_SHADER_PARAMETER_STRUCT(FAutoExposureTestUploadParameters, )
+	RDG_TEXTURE_ACCESS(InputTexture, ERHIAccess::CopyDest)
+END_SHADER_PARAMETER_STRUCT()
+
+BEGIN_SHADER_PARAMETER_STRUCT(FAutoExposureTestDownloadParameters, )
+	RDG_BUFFER_ACCESS(OutputBuffer, ERHIAccess::CopySrc)
+END_SHADER_PARAMETER_STRUCT()
+
+IMPLEMENT_SIMPLE_AUTOMATION_TEST(FAutoExposureTest, "NNEDenoiser.UnitTests.AutoExposure", EAutomationTestFlags_ApplicationContextMask | EAutomationTestFlags_FeatureMask | EAutomationTestFlags::EngineFilter)
+bool FAutoExposureTest::RunTest(const FString& Parameter)
+{
+	const int32 Width = 1920;
+	const int32 Height = 1088;
+
+	TArray<FLinearColor> TestData;
+	TestData.SetNumUninitialized(Width * Height);
+
+	for (int i = 0; i < TestData.Num(); i++)
+	{
+		TestData[i] = FLinearColor(FMath::FRand(), FMath::FRand(), FMath::FRand(), 1.0);
+	}
+
+	FAutoExposure AutoExposure;
+
+	float ResultCPU = 0.0f;
+	AutoExposure.Run(TestData, {Width, Height}, ResultCPU);
+
+	TArray<float> ResultRDG;
+	ResultRDG.SetNumZeroed(2);
+
+	FEvent* Signal = FGenericPlatformProcess::GetSynchEventFromPool(true);
+
+	ENQUEUE_RENDER_COMMAND(NNEAutoExposureTest)
+	(
+		[&TestData, &ResultRDG, AutoExposure, Signal](FRHICommandListImmediate& RHICmdList) mutable
+		{
+			ERHIPipeline Pipeline = RHICmdList.GetPipeline();
+			if (Pipeline == ERHIPipeline::None)
+			{
+				RHICmdList.SwitchPipeline(ERHIPipeline::Graphics);
+			}
+
+			FRDGBuilder	GraphBuilder(RHICmdList);
+
+			const FRDGTextureDesc TextureDesc = FRDGTextureDesc::Create2D({Width, Height}, EPixelFormat::PF_A32B32G32R32F, FClearValueBinding::Black, ETextureCreateFlags::None);
+			FRDGTextureRef InputTexture = GraphBuilder.CreateTexture(TextureDesc, TEXT("AutoExposureInputTexture"));
+
+			{
+				FAutoExposureTestUploadParameters* Parameters = GraphBuilder.AllocParameters<FAutoExposureTestUploadParameters>();
+				Parameters->InputTexture = InputTexture;
+
+				GraphBuilder.AddPass(RDG_EVENT_NAME("NNEDenoiser.AutoExposureTest.Upload"), Parameters, ERDGPassFlags::Readback,
+				[&TestData, Parameters](FRHICommandListImmediate& RHICmdList)
+				{
+					SCOPED_NAMED_EVENT_TEXT("FAutoExposureTest::Upload", FColor::Magenta);
+
+					CopyTextureFromCPUToGPU(RHICmdList, TestData, {Width, Height}, Parameters->InputTexture->GetRHI());
+				});
+			}
+
+			FRDGBufferDesc InputBufferDesc = FRDGBufferDesc::CreateBufferDesc(sizeof(float), 2);
+			FRDGBufferRef OutputBuffer = GraphBuilder.CreateBuffer(InputBufferDesc, TEXT("AutoExposureOutputBuffer"));
+
+			AutoExposure.EnqueueRDG(GraphBuilder, InputTexture, OutputBuffer);
+
+			{
+				FAutoExposureTestDownloadParameters* Parameters = GraphBuilder.AllocParameters<FAutoExposureTestDownloadParameters>();
+				Parameters->OutputBuffer = OutputBuffer;
+
+				GraphBuilder.AddPass(RDG_EVENT_NAME("NNEDenoiser.AutoExposureTest.Download"), Parameters, ERDGPassFlags::Readback,
+				[&ResultRDG, Parameters](FRHICommandListImmediate& RHICmdList)
+				{
+					SCOPED_NAMED_EVENT_TEXT("FAutoExposureTest::Download", FColor::Magenta);
+
+					CopyBufferFromGPUToCPU(RHICmdList, Parameters->OutputBuffer->GetRHI(), 2, ResultRDG);
+				});
+			}
+
+			GraphBuilder.Execute();
+
+			RHICmdList.BlockUntilGPUIdle();
+
+			Signal->Trigger();
+		}
+	);
+
+	Signal->Wait();
+
+	UTEST_EQUAL(TEXT("AutoExposure"), ResultRDG[0], ResultCPU);
+
+	return true;
+}
+
+#endif //WITH_DEV_AUTOMATION_TESTS
+
+} // namespace UE::NNEDenoiser::Internal
diff --git a/Engine/Plugins/NNE/NNEDenoiser/Source/NNEDenoiser/Private/NNEDenoiserUtils.h b/Engine/Plugins/NNE/NNEDenoiser/Source/NNEDenoiser/Private/NNEDenoiserUtils.h
new file mode 100644
--- /dev/null
+++ b/Engine/Plugins/NNE/NNEDenoiser/Source/NNEDenoiser/Private/NNEDenoiserUtils.h
@@ -0,0 +1,168 @@
+// Copyright Epic Games, Inc. All Rights Reserved.
+
+#pragma once
+
+#include "NNEDenoiserLog.h"
+#include "NNEDenoiserParameters.h"
+#include "NNEDenoiserShadersMappedCopyCS.h"
+#include "NNETypes.h"
+#include "RHICommandList.h"
+#include "RHIGPUReadback.h"
+#include "RHIResources.h"
+#include "RHITypes.h"
+
+namespace UE::NNEDenoiser::Private
+{
+	// Returns ceil(a / b) for non-negative integers
+	template<typename Int, typename IntB>
+	Int CeilDiv(Int a, IntB b)
+	{
+		return (a + b - 1) / b;
+	}
+
+	// Returns a rounded up to multiple of b
+	template<typename Int, typename IntB>
+	Int RoundUp(Int a, IntB b)
+	{
+		return CeilDiv(a, b) * b;
+	}
+
+	template<class IntType>
+	bool IsTensorShapeValid(TConstArrayView<IntType> ShapeData, TConstArrayView<int32> RequiredShapeData, const FString& Label)
+	{
+		static_assert(std::is_integral_v<IntType>);
+
+		if (ShapeData.Num() != RequiredShapeData.Num())
+		{
+			UE_LOG(LogNNEDenoiser, Error, TEXT("%s has wrong rank (expected %d, got %d)!"), *Label, RequiredShapeData.Num(), ShapeData.Num())
+			return false;
+		}
+
+		for (int32 I = 0; I < RequiredShapeData.Num(); I++)
+		{
+			if (RequiredShapeData[I] >= 0 && (int32)ShapeData[I] != RequiredShapeData[I])
+			{
+				UE_LOG(LogNNEDenoiser, Error, TEXT("%s does not have required shape (expected %d, got %d)!"), *Label, RequiredShapeData[I], ShapeData[I])
+				return false;
+			}
+		}
+
+		return true;
+	}
+
+	template <typename PixelType>
+	void CopyTextureFromGPUToCPU(FRHICommandListImmediate& RHICmdList, FRHITexture* SrcTexture, FIntPoint Size, TArray<PixelType>& DstArray)
+	{
+		SCOPED_NAMED_EVENT_TEXT("NNEDenoiser.CopyTextureFromGPUToCPU", FColor::Magenta);
+
+		FRHIGPUTextureReadback Readback(TEXT("NNEDenoiser.CopyTextureFromGPUToCPU"));
+		Readback.EnqueueCopy(RHICmdList, SrcTexture, FIntVector::ZeroValue, 0, FIntVector(Size.X, Size.Y, 1));
+
+		RHICmdList.BlockUntilGPUIdle();
+
+		int32_t SrcStride = 0;
+		const PixelType* SrcBuffer = static_cast<PixelType*>(Readback.Lock(SrcStride, nullptr));
+
+		PixelType* DstBuffer = DstArray.GetData();
+		for (int Y = 0; Y < Size.Y; Y++, DstBuffer += Size.X, SrcBuffer += SrcStride)
+		{
+			FPlatformMemory::Memcpy(DstBuffer, SrcBuffer, Size.X * sizeof(PixelType));
+
+		}
+		Readback.Unlock();
+	}
+
+	template <typename PixelType>
+	void CopyTextureFromCPUToGPU(FRHICommandListImmediate& RHICmdList, const TArray<PixelType>& SrcArray, FIntPoint Size, FRHITexture* DstTexture)
+	{
+		SCOPED_NAMED_EVENT_TEXT("NNEDenoiser.CopyTextureFromCPUToGPU", FColor::Magenta);
+
+		uint32 DestStride;
+		PixelType* DstBuffer = static_cast<PixelType*>(RHICmdList.LockTexture2D(DstTexture, 0, RLM_WriteOnly, DestStride, false));
+		DestStride /= sizeof(PixelType);
+		const PixelType* SrcBuffer = SrcArray.GetData();
+		for (int32 Y = 0; Y < Size.Y; Y++, SrcBuffer += Size.X, DstBuffer += DestStride)
+		{
+			FPlatformMemory::Memcpy(DstBuffer, SrcBuffer, Size.X * sizeof(PixelType));
+		}
+		RHICmdList.UnlockTexture2D(DstTexture, 0, false);
+	}
+
+	template <typename ElementType>
+	void CopyBufferFromGPUToCPU(FRHICommandListImmediate& RHICmdList, FRHIBuffer* SrcBuffer, int32 Count, TArray<ElementType>& DstArray)
+	{
+		SCOPED_NAMED_EVENT_TEXT("NNEDenoiser.CopyBufferFromGPUToCPU", FColor::Magenta);
+
+		const int32 NumBytes = Count * sizeof(ElementType);
+
+		FRHIGPUBufferReadback Readback(TEXT("NNEDenoiser.CopyBufferFromGPUToCPU"));
+		Readback.EnqueueCopy(RHICmdList, SrcBuffer, NumBytes);
+
+		RHICmdList.BlockUntilGPUIdle();
+
+		const ElementType* Src = static_cast<ElementType*>(Readback.Lock(NumBytes));
+		ElementType* Dst = DstArray.GetData();
+		
+		FPlatformMemory::Memcpy(Dst, Src, NumBytes);
+
+		Readback.Unlock();
+	}
+
+	template <typename ElementType>
+	void CopyBufferFromCPUToGPU(FRHICommandListImmediate& RHICmdList, const TArray<ElementType>& SrcArray, int32 Count, FRHIBuffer* DstBuffer)
+	{
+		SCOPED_NAMED_EVENT_TEXT("NNEDenoiser.CopyBufferFromCPUToGPU", FColor::Magenta);
+
+		ElementType* Dst = static_cast<ElementType*>(RHICmdList.LockBuffer(DstBuffer, 0, Count * sizeof(ElementType), RLM_WriteOnly));
+		const ElementType* Src = SrcArray.GetData();
+		
+		FPlatformMemory::Memcpy(Dst, Src, Count * sizeof(ElementType));
+
+		RHICmdList.UnlockBuffer(DstBuffer);
+	}
+
+	inline EPixelFormat GetBufferFormat(ENNETensorDataType TensorDataType)
+	{
+		switch (TensorDataType)
+		{
+			case ENNETensorDataType::Half: return EPixelFormat::PF_R16F;
+			case ENNETensorDataType::Float: return EPixelFormat::PF_R32_FLOAT;
+			default:
+				unimplemented();
+		}
+		return EPixelFormat::PF_Unknown;
+	}
+
+	inline UE::NNEDenoiserShaders::Internal::EDataType GetDenoiserShaderDataType(ENNETensorDataType TensorDataType)
+	{
+		using UE::NNEDenoiserShaders::Internal::EDataType;
+
+		switch (TensorDataType)
+		{
+			case ENNETensorDataType::Half: return EDataType::Half;
+			case ENNETensorDataType::Float: return EDataType::Float;
+			default:
+				unimplemented();
+		}
+		return EDataType::None;
+	}
+
+	inline UE::NNEDenoiserShaders::Internal::EDataType GetDenoiserShaderDataType(EPixelFormat Format)
+	{
+		using UE::NNEDenoiserShaders::Internal::EDataType;
+
+		switch (Format)
+		{
+			case EPixelFormat::PF_R16F:
+			case EPixelFormat::PF_FloatRGBA: return EDataType::Half;
+
+			case EPixelFormat::PF_R32_FLOAT:
+			case EPixelFormat::PF_A32B32G32R32F: return EDataType::Float;
+			
+			default:
+				unimplemented();
+		}
+		return EDataType::None;
+	} 
+
+}
\ No newline at end of file
diff --git a/Engine/Plugins/NNE/NNEDenoiser/Source/NNEDenoiser/Private/NNEDenoiserViewExtension.cpp b/Engine/Plugins/NNE/NNEDenoiser/Source/NNEDenoiser/Private/NNEDenoiserViewExtension.cpp
new file mode 100644
--- /dev/null
+++ b/Engine/Plugins/NNE/NNEDenoiser/Source/NNEDenoiser/Private/NNEDenoiserViewExtension.cpp
@@ -0,0 +1,434 @@
+// Copyright Epic Games, Inc. All Rights Reserved.
+
+#include "NNEDenoiserViewExtension.h"
+#include "NNE.h"
+#include "NNEDenoiserAsset.h"
+#include "NNEDenoiserAutoExposure.h"
+#include "NNEDenoiserGenericDenoiser.h"
+#include "NNEDenoiserIOMappingData.h"
+#include "NNEDenoiserIOProcessBase.h"
+#include "NNEDenoiserLog.h"
+#include "NNEDenoiserModelInstanceCPU.h"
+#include "NNEDenoiserModelInstanceGPU.h"
+#include "NNEDenoiserModelInstanceRDG.h"
+#include "NNEDenoiserPathTracingDenoiser.h"
+#include "NNEDenoiserPathTracingSpatialTemporalDenoiser.h"
+#include "NNEDenoiserResourceMapping.h"
+#include "NNEDenoiserTemporalAsset.h"
+#include "NNEDenoiserTransferFunctionOidn.h"
+#include "NNEDenoiserUtils.h"
+#include "NNEModelData.h"
+#include "PathTracingDenoiser.h"
+
+namespace UE::NNEDenoiser::Private
+{
+
+static TAutoConsoleVariable<bool> CVarNNEDenoiser(
+	TEXT("NNEDenoiser"),
+	true,
+	TEXT("Enable the NNE Denoiser.")
+);
+
+static TAutoConsoleVariable<int32> CVarNNEDenoiserAsset(
+	TEXT("NNEDenoiser.Asset"),
+	0,
+	TEXT("Defines denoiser asset to used to create the denoiser.\n")
+	TEXT("  0: Use denoiser asset defined by Project Settings\n")
+	TEXT("  1: OIDN Fast\n")
+	TEXT("  2: OIDN Balanced\n")
+	TEXT("  3: OIDN High Quality\n")
+	TEXT("  4: OIDN Fast | Alpha\n")
+	TEXT("  5: OIDN Balanced | Alpha\n")
+	TEXT("  6: OIDN High Quality | Alpha")
+);
+
+static TAutoConsoleVariable<int32> CVarNNEDenoiserTemporalAsset(
+	TEXT("NNEDenoiser.TemporalAsset"),
+	0,
+	TEXT("Defines temporal denoiser asset to used to create the temporal denoiser.\n")
+	TEXT("  0: Use temporal denoiser asset defined by Project Settings\n")
+	TEXT("  1: OIDN Balanced\n")
+	TEXT("  2: OIDN Balanced | Alpha")
+);
+
+static TAutoConsoleVariable<int32> CVarNNEDenoiserRuntimeType(
+	TEXT("NNEDenoiser.Runtime.Type"),
+	2,
+	TEXT("Defines the runtime type to run the denoiser model.\n")
+	TEXT("  0: CPU\n")
+	TEXT("  1: GPU\n")
+	TEXT("  2: RDG\n")
+);
+
+static TAutoConsoleVariable<FString> CVarNNEDenoiserRuntimeName(
+	TEXT("NNEDenoiser.Runtime.Name"),
+	FString(),
+	TEXT("Defines the runtime name to run the denoiser model. Leave empty to use default.")
+);
+
+EDenoiserRuntimeType GetDenoiserRuntimeTypeFromCVar()
+{
+	const int32 Value = CVarNNEDenoiserRuntimeType.GetValueOnGameThread();
+	const int32 Min = (int32)EDenoiserRuntimeType::CPU;
+	const int32 Max = (int32)EDenoiserRuntimeType::RDG;
+	return static_cast<EDenoiserRuntimeType>(FMath::Clamp(Value, Min, Max));
+}
+
+FString GetDenoiserAssetNameFromCVarAndSettings(const UNNEDenoiserSettings* Settings)
+{
+	const int32 Idx = FMath::Clamp(CVarNNEDenoiserAsset.GetValueOnGameThread(), 0, 6);
+	switch(Idx)
+	{
+		case 0: return !Settings->DenoiserAsset.IsNull() ? Settings->DenoiserAsset.ToString() : FString();
+
+		case 1: return TEXT("/NNEDenoiser/NNED_Oidn2-3_Fast.NNED_Oidn2-3_Fast");
+		case 2: return TEXT("/NNEDenoiser/NNED_Oidn2-3_Balanced.NNED_Oidn2-3_Balanced");
+		case 3: return TEXT("/NNEDenoiser/NNED_Oidn2-3_HighQuality.NNED_Oidn2-3_HighQuality");
+
+		// Alpha
+		case 4: return TEXT("/NNEDenoiser/NNED_Oidn2-3_Fast_Alpha.NNED_Oidn2-3_Fast_Alpha");
+		case 5: return TEXT("/NNEDenoiser/NNED_Oidn2-3_Balanced_Alpha.NNED_Oidn2-3_Balanced_Alpha");
+		case 6: return TEXT("/NNEDenoiser/NNED_Oidn2-3_HighQuality_Alpha.NNED_Oidn2-3_HighQuality_Alpha");
+	}
+	check(false);
+	return FString();
+}
+
+FString GetDenoiserTemporalAssetNameFromCVarAndSettings(const UNNEDenoiserSettings* Settings)
+{
+	const int32 Idx = FMath::Clamp(CVarNNEDenoiserTemporalAsset.GetValueOnGameThread(), 0, 6);
+	switch(Idx)
+	{
+		case 0: return !Settings->TemporalDenoiserAsset.IsNull() ? Settings->TemporalDenoiserAsset.ToString() : FString();
+
+		case 1: return TEXT("/NNEDenoiser/NNEDT_Oidn2-3_Balanced.NNEDT_Oidn2-3_Balanced");
+
+		// Alpha
+		case 2: return TEXT("/NNEDenoiser/NNEDT_Oidn2-3_Balanced_Alpha.NNEDT_Oidn2-3_Balanced_Alpha");
+	}
+	check(false);
+	return FString();
+}
+
+TUniquePtr<FGenericDenoiser> CreateNNEDenoiser(
+	UNNEModelData& ModelData,
+	EDenoiserRuntimeType RuntimeType,
+	const FString& RuntimeNameOverride,
+	TUniquePtr<IInputProcess> InputProcess,
+	TUniquePtr<IOutputProcess> OutputProcess,
+	FParameters Parameters,
+	TUniquePtr<IAutoExposure> AutoExposure,
+	TSharedPtr<ITransferFunction> TransferFunction)
+{
+	static TMap<EDenoiserRuntimeType, FString> DefaultRuntimeNames;
+	DefaultRuntimeNames.Emplace(EDenoiserRuntimeType::CPU, TEXT("NNERuntimeORTCpu"));
+	DefaultRuntimeNames.Emplace(EDenoiserRuntimeType::GPU, TEXT("NNERuntimeORTDml"));
+	DefaultRuntimeNames.Emplace(EDenoiserRuntimeType::RDG, TEXT("NNERuntimeORTDml"));
+
+	const static FString FallbackRuntimeNameRDG = TEXT("NNERuntimeRDGHlsl");
+
+	TArray<TPair<EDenoiserRuntimeType, FString>> RuntimePriorityQueue;
+	if (!RuntimeNameOverride.IsEmpty())
+	{
+		RuntimePriorityQueue.AddUnique({RuntimeType, RuntimeNameOverride});
+	}
+	
+	const FString& DefaultRuntimeName = DefaultRuntimeNames.FindChecked(RuntimeType);
+	if (RuntimeNameOverride != DefaultRuntimeName)
+	{
+		RuntimePriorityQueue.AddUnique({RuntimeType, DefaultRuntimeName});
+	}
+
+	if (RuntimeType >= EDenoiserRuntimeType::RDG)
+	{
+		if (RuntimeNameOverride != FallbackRuntimeNameRDG)
+		{
+			RuntimePriorityQueue.AddUnique({EDenoiserRuntimeType::RDG, FallbackRuntimeNameRDG});
+		}
+
+		const FString& DefaultRuntimeNameGPU = DefaultRuntimeNames.FindChecked(EDenoiserRuntimeType::GPU);
+		RuntimePriorityQueue.AddUnique({EDenoiserRuntimeType::GPU, DefaultRuntimeNameGPU});
+	}
+
+	if (RuntimeType >= EDenoiserRuntimeType::GPU)
+	{
+		const FString& DefaultRuntimeNameCPU = DefaultRuntimeNames.FindChecked(EDenoiserRuntimeType::CPU);
+		RuntimePriorityQueue.AddUnique({EDenoiserRuntimeType::CPU, DefaultRuntimeNameCPU});
+	}
+
+	TUniquePtr<IModelInstance> ModelInstance;
+	for (const auto& Pair : RuntimePriorityQueue)
+	{
+		UE_LOG(LogNNEDenoiser, Log, TEXT("Try create model instance with runtime %s on %s..."), *Pair.Value, *UEnum::GetValueAsString(Pair.Key));
+
+		if (Pair.Key == EDenoiserRuntimeType::CPU)
+		{
+			ModelInstance = FModelInstanceCPU::Make(ModelData, Pair.Value);
+		}
+		else if (Pair.Key == EDenoiserRuntimeType::GPU)
+		{
+			ModelInstance = FModelInstanceGPU::Make(ModelData, Pair.Value);
+		}
+		else if (Pair.Key == EDenoiserRuntimeType::RDG)
+		{
+			ModelInstance = FModelInstanceRDG::Make(ModelData, Pair.Value);
+		}
+
+		if (ModelInstance.IsValid())
+		{
+			UE_LOG(LogNNEDenoiser, Display, TEXT("Created model instance with runtime %s on %s"), *Pair.Value, *UEnum::GetValueAsString(Pair.Key));
+			break;
+		}
+	}
+
+	if (!ModelInstance.IsValid())
+	{
+		UE_LOG(LogNNEDenoiser, Error, TEXT("Could not create denoiser!"));
+		return {};
+	}
+
+	return MakeUnique<FGenericDenoiser>(
+		MoveTemp(ModelInstance),
+		MoveTemp(InputProcess),
+		MoveTemp(OutputProcess),
+		MoveTemp(Parameters),
+		MoveTemp(AutoExposure),
+		TransferFunction);
+}
+
+template<class AssetType>
+FParameters GetParametersValidated(const AssetType& DenoiserAsset)
+{
+	FParameters Parameters{
+		.TilingConfig =
+		{
+			.Alignment = DenoiserAsset.TilingConfig.Alignment,
+			.Overlap = DenoiserAsset.TilingConfig.Overlap,
+			.MaxSize = DenoiserAsset.TilingConfig.MaxSize,
+			.MinSize = DenoiserAsset.TilingConfig.MinSize
+		}
+	};
+
+	if (Parameters.TilingConfig.Alignment < 1)
+	{
+		UE_LOG(LogNNEDenoiser, Warning, TEXT("Tiling alignment should be at least 1!"));
+		Parameters.TilingConfig.Alignment = 1;
+	}
+
+	if (Parameters.TilingConfig.Overlap % Parameters.TilingConfig.Alignment != 0)
+	{
+		UE_LOG(LogNNEDenoiser, Warning, TEXT("Tiling overlap should be aligned by %d!"), Parameters.TilingConfig.Alignment);
+		Parameters.TilingConfig.Overlap = RoundUp(Parameters.TilingConfig.Overlap, Parameters.TilingConfig.Alignment);
+	}
+
+	if (Parameters.TilingConfig.MinSize < Parameters.TilingConfig.Overlap + Parameters.TilingConfig.Alignment)
+	{
+		UE_LOG(LogNNEDenoiser, Warning, TEXT("Minimum tile size should be at least overlap + alignment = %d!"), Parameters.TilingConfig.Overlap + Parameters.TilingConfig.Alignment);
+		Parameters.TilingConfig.MinSize = Parameters.TilingConfig.Overlap + Parameters.TilingConfig.Alignment;
+	}
+
+	if (Parameters.TilingConfig.MinSize % Parameters.TilingConfig.Alignment != 0)
+	{
+		UE_LOG(LogNNEDenoiser, Warning, TEXT("Minimum tile size should be aligned by %d!"), Parameters.TilingConfig.Alignment);
+		Parameters.TilingConfig.MinSize = RoundUp(Parameters.TilingConfig.MinSize, Parameters.TilingConfig.Alignment);
+	}
+
+	if (Parameters.TilingConfig.MaxSize > 0 && Parameters.TilingConfig.MaxSize < Parameters.TilingConfig.MinSize)
+	{
+		UE_LOG(LogNNEDenoiser, Warning, TEXT("Maximum tile size should be at least minimum tile size!"));
+		Parameters.TilingConfig.MaxSize = Parameters.TilingConfig.MinSize;
+	}
+
+	if (Parameters.TilingConfig.MaxSize % Parameters.TilingConfig.Alignment != 0)
+	{
+		UE_LOG(LogNNEDenoiser, Warning, TEXT("Maximum tile size should be aligned by %d!"), Parameters.TilingConfig.Alignment);
+		Parameters.TilingConfig.MaxSize = RoundUp(Parameters.TilingConfig.MaxSize, Parameters.TilingConfig.Alignment);
+	}
+
+	return Parameters;
+}
+
+template<class AssetType, class InputMappingType, class OutputMappingType>
+TUniquePtr<FGenericDenoiser> CreateNNEDenoiserFromAsset(const FString& AssetName, EDenoiserRuntimeType RuntimeType, const FString& RuntimeNameOverride)
+{
+	if (AssetName.IsEmpty())
+	{
+		UE_LOG(LogNNEDenoiser, Error, TEXT("Asset name not set!"));
+		return {};
+	}
+
+	AssetType *DenoiserAsset = LoadObject<AssetType>(nullptr, *AssetName);
+	if (!DenoiserAsset)
+	{
+		UE_LOG(LogNNEDenoiser, Error, TEXT("Could not load denoiser model data asset!"));
+		return {};
+	}
+
+	UNNEModelData* ModelData = DenoiserAsset->ModelData.LoadSynchronous();
+	if (!ModelData)
+	{
+		UE_LOG(LogNNEDenoiser, Error, TEXT("Asset does not contain model data!"));
+		return {};
+	}
+	UDataTable* InputMappingTable = DenoiserAsset->InputMapping.LoadSynchronous();
+	if (InputMappingTable)
+	{
+		UE_LOG(LogNNEDenoiser, Log, TEXT("Loaded input mapping from %s"), *DenoiserAsset->InputMapping.GetAssetName());
+	}
+	UDataTable* OutputMappingTable = DenoiserAsset->OutputMapping.LoadSynchronous();
+	if (OutputMappingTable)
+	{
+		UE_LOG(LogNNEDenoiser, Log, TEXT("Loaded output mapping from %s"), *DenoiserAsset->OutputMapping.GetAssetName());
+	}
+
+	FResourceMappingList InputLayout;
+	if (InputMappingTable)
+	{
+		InputLayout = MakeTensorLayout<InputMappingType>(InputMappingTable);
+	}
+
+	FResourceMappingList OutputLayout;
+	if (OutputMappingTable)
+	{
+		OutputLayout = MakeTensorLayout<OutputMappingType>(OutputMappingTable);
+	}
+
+	TUniquePtr<FAutoExposure> AutoExposure;
+	TSharedPtr<ITransferFunction> TransferFunction;
+	if (AssetName.Contains(TEXT("oidn2")))
+	{
+		AutoExposure = MakeUnique<FAutoExposure>();
+		TransferFunction = MakeShared<Oidn::FTransferFunction>();
+	}
+
+	TUniquePtr<IInputProcess> InputProcess = MakeUnique<FInputProcessBase>(MoveTemp(InputLayout), TransferFunction);
+	TUniquePtr<IOutputProcess> OutputProcess = MakeUnique<FOutputProcessBase>(MoveTemp(OutputLayout), TransferFunction);
+
+	FParameters Parameters = GetParametersValidated(*DenoiserAsset);
+
+	return CreateNNEDenoiser(
+		*ModelData,
+		RuntimeType,
+		RuntimeNameOverride,
+		MoveTemp(InputProcess),
+		MoveTemp(OutputProcess),
+		MoveTemp(Parameters),
+		MoveTemp(AutoExposure),
+		TransferFunction);
+}
+
+FViewExtension::FViewExtension(const FAutoRegister& AutoRegister) : FSceneViewExtensionBase(AutoRegister)
+{
+#if WITH_EDITOR
+	GetMutableDefault<UNNEDenoiserSettings>()->OnSettingChanged().AddRaw(this, &FViewExtension::OnDenoiserSettingsChanged);
+#endif
+}
+
+FViewExtension::~FViewExtension()
+{
+#if WITH_EDITOR
+	if(UObjectInitialized())
+	{
+		GetMutableDefault<UNNEDenoiserSettings>()->OnSettingChanged().RemoveAll(this);
+	}
+#endif
+
+	UnregisterDenoiser(TEXT("NNEDenoiser"));
+}
+
+void FViewExtension::BeginRenderViewFamily(FSceneViewFamily& InViewFamily)
+{
+	ApplySettings(GetDefault<UNNEDenoiserSettings>());
+}
+
+void FViewExtension::PreRenderViewFamily_RenderThread(FRDGBuilder& GraphBuilder, FSceneViewFamily& InViewFamily)
+{
+	if (!bDenoiserEnabled)
+	{
+		UnregisterDenoiser(TEXT("NNEDenoiser"));
+	}
+	else
+	{
+		if (DenoiserToSwap.IsValid() || SpatialTemporalDenoiserToSwap.IsValid())
+		{
+			UnregisterDenoiser(TEXT("NNEDenoiser"));
+		}
+
+		if (DenoiserToSwap.IsValid())
+		{
+			RegisterSpatialDenoiser(MoveTemp(DenoiserToSwap),TEXT("NNEDenoiser"));
+		}
+
+		if (SpatialTemporalDenoiserToSwap.IsValid())
+		{
+			RegisterSpatialTemporalDenoiser(MoveTemp(SpatialTemporalDenoiserToSwap), TEXT("NNEDenoiser"));
+		}
+	}
+}
+
+void FViewExtension::ApplySettings(const UNNEDenoiserSettings* Settings)
+{
+	// Check for changes
+	bool bNeedsUpdate = false;
+	bNeedsUpdate |= CVarNNEDenoiser.GetValueOnGameThread() != bDenoiserEnabled;
+	bNeedsUpdate |= GetDenoiserRuntimeTypeFromCVar() != RuntimeType;
+	bNeedsUpdate |= CVarNNEDenoiserRuntimeName.GetValueOnGameThread() != RuntimeName;
+	bNeedsUpdate |= GetDenoiserAssetNameFromCVarAndSettings(Settings) != AssetName;
+	bNeedsUpdate |= GetDenoiserTemporalAssetNameFromCVarAndSettings(Settings) != TemporalAssetName;
+
+	if (!bNeedsUpdate)
+	{
+		return;
+	}
+
+	bDenoiserEnabled = CVarNNEDenoiser.GetValueOnGameThread();
+
+	UE_LOG(LogNNEDenoiser, Log, TEXT("ApplySettings: bDenoiserEnabled %d"), bDenoiserEnabled);
+
+	if (!bDenoiserEnabled)
+	{
+		return;
+	}
+
+	RuntimeType = GetDenoiserRuntimeTypeFromCVar();
+	RuntimeName = CVarNNEDenoiserRuntimeName.GetValueOnGameThread();
+	AssetName = GetDenoiserAssetNameFromCVarAndSettings(Settings);
+	TemporalAssetName = GetDenoiserTemporalAssetNameFromCVarAndSettings(Settings);
+
+	// UE_LOG(LogNNEDenoiser, Log, TEXT("Create denoiser from asset %s..."), *AssetName);
+
+	TUniquePtr<FGenericDenoiser> Denoiser = CreateNNEDenoiserFromAsset<UNNEDenoiserAsset, FNNEDenoiserInputMappingData, FNNEDenoiserOutputMappingData>(AssetName, RuntimeType, RuntimeName);
+	if (Denoiser.IsValid())
+	{
+		UE_LOG(LogNNEDenoiser, Log, TEXT("Create denoiser from asset %s..."), *AssetName);
+
+		DenoiserToSwap = MakeUnique<FPathTracingDenoiser>(MoveTemp(Denoiser));
+	}
+	else
+	{
+		UE_LOG(LogNNEDenoiser, Error, TEXT("Could not create denoiser from asset %d!"), *AssetName);
+	}
+
+	TUniquePtr<FGenericDenoiser> TemporalDenoiser = CreateNNEDenoiserFromAsset<UNNEDenoiserTemporalAsset, FNNEDenoiserTemporalInputMappingData, FNNEDenoiserTemporalOutputMappingData>(TemporalAssetName, RuntimeType, RuntimeName);
+	if (TemporalDenoiser.IsValid())
+	{
+		UE_LOG(LogNNEDenoiser, Log, TEXT("Create temporal denoiser from asset %s..."), *TemporalAssetName);
+
+		SpatialTemporalDenoiserToSwap = MakeUnique<FPathTracingSpatialTemporalDenoiser>(MoveTemp(TemporalDenoiser));
+	}
+	else
+	{
+		UE_LOG(LogNNEDenoiser, Error, TEXT("Could not create temporal denoiser from asset %d!"), *TemporalAssetName);
+	}
+}
+
+#if WITH_EDITOR
+void FViewExtension::OnDenoiserSettingsChanged(UObject* InObject, struct FPropertyChangedEvent& InPropertyChangedEvent)
+{
+	UE_LOG(LogNNEDenoiser, Log, TEXT("Settings %s changed: %s"), *InObject->GetName(), *InPropertyChangedEvent.GetPropertyName().ToString());
+
+	ApplySettings(CastChecked<UNNEDenoiserSettings>(InObject));
+}
+#endif
+
+} // namespace UE::NNEDenoiser::Private
\ No newline at end of file
diff --git a/Engine/Plugins/NNE/NNEDenoiser/Source/NNEDenoiser/Private/NNEDenoiserViewExtension.h b/Engine/Plugins/NNE/NNEDenoiser/Source/NNEDenoiser/Private/NNEDenoiserViewExtension.h
new file mode 100644
--- /dev/null
+++ b/Engine/Plugins/NNE/NNEDenoiser/Source/NNEDenoiser/Private/NNEDenoiserViewExtension.h
@@ -0,0 +1,48 @@
+// Copyright Epic Games, Inc. All Rights Reserved.
+
+#pragma once
+
+#include "SceneViewExtension.h"
+#include "NNEDenoiserSettings.h"
+
+namespace UE::Renderer::Private
+{
+	class IPathTracingDenoiser;
+	class IPathTracingSpatialTemporalDenoiser;
+}
+
+namespace UE::NNEDenoiser::Private
+{
+
+class FViewExtension final : public FSceneViewExtensionBase
+{
+	public:
+		FViewExtension(const FAutoRegister& AutoRegister);
+		virtual ~FViewExtension();
+
+		void SetupViewFamily(FSceneViewFamily& InViewFamily) override { }
+		void SetupView(FSceneViewFamily& InViewFamily, FSceneView& InView) override { }
+			
+		void BeginRenderViewFamily(FSceneViewFamily& InViewFamily) override;
+
+		void PreRenderViewFamily_RenderThread(FRDGBuilder& GraphBuilder, FSceneViewFamily& InViewFamily) override;
+
+	private:
+		void ApplySettings(const UNNEDenoiserSettings* Settings);
+
+#if WITH_EDITOR
+		void OnDenoiserSettingsChanged(UObject* InObject, struct FPropertyChangedEvent& InPropertyChangedEvent);
+#endif
+
+		TUniquePtr<UE::Renderer::Private::IPathTracingDenoiser> DenoiserToSwap;
+		TUniquePtr<UE::Renderer::Private::IPathTracingSpatialTemporalDenoiser> SpatialTemporalDenoiserToSwap;
+
+		// Cached settings and CVar values
+		bool bDenoiserEnabled = true;
+		EDenoiserRuntimeType RuntimeType;
+		FString RuntimeName;
+		FString AssetName;
+		FString TemporalAssetName;
+	};
+
+} // namespace UE::NNEDenoiser::Private
\ No newline at end of file
diff --git a/Engine/Plugins/NNE/NNEDenoiser/Source/NNEDenoiser/Public/NNEDenoiserAsset.h b/Engine/Plugins/NNE/NNEDenoiser/Source/NNEDenoiser/Public/NNEDenoiserAsset.h
new file mode 100644
--- /dev/null
+++ b/Engine/Plugins/NNE/NNEDenoiser/Source/NNEDenoiser/Public/NNEDenoiserAsset.h
@@ -0,0 +1,34 @@
+// Copyright Epic Games, Inc. All Rights Reserved.
+
+#pragma once
+
+#include "Engine/DataAsset.h"
+#include "Engine/DataTable.h"
+#include "NNEDenoiserTilingConfig.h"
+#include "NNEModelData.h"
+
+#include "NNEDenoiserAsset.generated.h"
+
+/** Denoiser model data asset */
+UCLASS(BlueprintType)
+class NNEDENOISER_API UNNEDenoiserAsset : public UDataAsset
+{
+	GENERATED_BODY()
+
+public:
+	/** NNE model data */
+	UPROPERTY(EditAnywhere, BlueprintReadWrite, Category=NNEDenoiser)
+	TSoftObjectPtr<UNNEModelData> ModelData;
+
+	/** Input mapping table */
+	UPROPERTY(EditAnywhere, BlueprintReadWrite, Category=NNEDenoiser, meta = (RequiredAssetDataTags = "RowStructure=/Script/NNEDenoiser.NNEDenoiserInputMappingData"))
+	TSoftObjectPtr<UDataTable> InputMapping;
+
+	/** Output mapping table */
+	UPROPERTY(EditAnywhere, BlueprintReadWrite, Category=NNEDenoiser, meta = (RequiredAssetDataTags = "RowStructure=/Script/NNEDenoiser.NNEDenoiserOutputMappingData"))
+	TSoftObjectPtr<UDataTable> OutputMapping;
+
+	/** Tiling configuration */
+	UPROPERTY(EditAnywhere, BlueprintReadWrite, Category=NNEDenoiser)
+	FTilingConfig TilingConfig{};
+};
\ No newline at end of file
diff --git a/Engine/Plugins/NNE/NNEDenoiser/Source/NNEDenoiser/Public/NNEDenoiserIOMappingData.h b/Engine/Plugins/NNE/NNEDenoiser/Source/NNEDenoiser/Public/NNEDenoiserIOMappingData.h
new file mode 100644
--- /dev/null
+++ b/Engine/Plugins/NNE/NNEDenoiser/Source/NNEDenoiser/Public/NNEDenoiserIOMappingData.h
@@ -0,0 +1,122 @@
+// Copyright Epic Games, Inc. All Rights Reserved.
+
+#pragma once
+
+#include "Engine/DataTable.h"
+#include "NNEDenoiserResourceName.h"
+
+#include "NNEDenoiserIOMappingData.generated.h"
+
+/** An enum to represent resource names used for denoiser input mapping */
+UENUM()
+enum class EInputResourceName : uint8
+{
+	Color,
+	Albedo,
+	Normal,
+	Output
+};
+
+/** An enum to represent resource names used for denoiser output mapping */
+UENUM()
+enum class EOutputResourceName : uint8
+{
+	Output
+};
+
+/** An enum to represent resource names used for temporal denoiser input mapping */
+UENUM()
+enum class ETemporalInputResourceName : uint8
+{
+	Color,
+	Albedo,
+	Normal,
+	Flow,
+	Output
+};
+
+/** An enum to represent resource names used for temporal denoiser output mapping */
+UENUM()
+enum class ETemporalOutputResourceName : uint8
+{
+	Output
+};
+
+/** Table row base for denoiser basic input and output mapping */
+USTRUCT(BlueprintType)
+struct FNNEDenoiserBaseMappingData : public FTableRowBase
+{
+	GENERATED_USTRUCT_BODY()
+
+public:
+	/** Input/output tensor index */
+	UPROPERTY(EditAnywhere, BlueprintReadWrite, Category=NNEDenoiser)
+	int32 TensorIndex = 0;
+
+	/** Input/output tensor channel */
+	UPROPERTY(EditAnywhere, BlueprintReadWrite, Category=NNEDenoiser)
+	int32 TensorChannel = 0;
+
+	/** Resource channel */
+	UPROPERTY(EditAnywhere, BlueprintReadWrite, Category=NNEDenoiser)
+	int32 ResourceChannel = 0;
+};
+
+/** Table row base for denoiser input mapping */
+USTRUCT(BlueprintType)
+struct FNNEDenoiserInputMappingData : public FNNEDenoiserBaseMappingData
+{
+	GENERATED_USTRUCT_BODY()
+
+	/** Mapped resource name */
+	UPROPERTY(EditAnywhere, BlueprintReadWrite, Category=NNEDenoiser)
+	EInputResourceName Resource = EInputResourceName::Color;
+};
+
+/** Table row base for denoiser output mapping */
+USTRUCT(BlueprintType)
+struct FNNEDenoiserOutputMappingData : public FNNEDenoiserBaseMappingData
+{
+	GENERATED_USTRUCT_BODY()
+
+	/** Mapped resource name */
+	UPROPERTY(EditAnywhere, BlueprintReadWrite, Category=NNEDenoiser)
+	EOutputResourceName Resource = EOutputResourceName::Output;
+};
+
+/** Table row base for temporal denoiser input mapping */
+USTRUCT(BlueprintType)
+struct FNNEDenoiserTemporalInputMappingData : public FNNEDenoiserBaseMappingData
+{
+	GENERATED_USTRUCT_BODY()
+
+public:
+	/** Mapped resource name */
+	UPROPERTY(EditAnywhere, BlueprintReadWrite, Category=NNEDenoiser)
+	ETemporalInputResourceName Resource = ETemporalInputResourceName::Color;
+
+	/** Resource frame index */
+	UPROPERTY(EditAnywhere, BlueprintReadWrite, Category=NNEDenoiser)
+	int32 FrameIndex = 0;
+};
+
+/** Table row base for temporal denoiser output mapping */
+USTRUCT(BlueprintType)
+struct FNNEDenoiserTemporalOutputMappingData : public FNNEDenoiserBaseMappingData
+{
+	GENERATED_USTRUCT_BODY()
+
+	/** Mapped resource name */
+	UPROPERTY(EditAnywhere, BlueprintReadWrite, Category=NNEDenoiser)
+	ETemporalOutputResourceName Resource = ETemporalOutputResourceName::Output;
+};
+
+namespace UE::NNEDenoiser
+{
+
+EResourceName ToResourceName(EInputResourceName Name);
+EResourceName ToResourceName(EOutputResourceName Name);
+EResourceName ToResourceName(ETemporalInputResourceName Name);
+EResourceName ToResourceName(ETemporalOutputResourceName Name);
+
+} // namespace UE::NNEDenoiser
\ No newline at end of file
diff --git a/Engine/Plugins/NNE/NNEDenoiser/Source/NNEDenoiser/Public/NNEDenoiserResourceName.h b/Engine/Plugins/NNE/NNEDenoiser/Source/NNEDenoiser/Public/NNEDenoiserResourceName.h
new file mode 100644
--- /dev/null
+++ b/Engine/Plugins/NNE/NNEDenoiser/Source/NNEDenoiser/Public/NNEDenoiserResourceName.h
@@ -0,0 +1,16 @@
+// Copyright Epic Games, Inc. All Rights Reserved.
+
+#pragma once
+
+#include "NNEDenoiserResourceName.generated.h"
+
+/** An enum to represent resource names used for input and output mapping */
+UENUM()
+enum class EResourceName : uint8
+{
+	Color,
+	Albedo,
+	Normal,
+	Flow,
+	Output
+};
\ No newline at end of file
diff --git a/Engine/Plugins/NNE/NNEDenoiser/Source/NNEDenoiser/Public/NNEDenoiserSettings.h b/Engine/Plugins/NNE/NNEDenoiser/Source/NNEDenoiser/Public/NNEDenoiserSettings.h
new file mode 100644
--- /dev/null
+++ b/Engine/Plugins/NNE/NNEDenoiser/Source/NNEDenoiser/Public/NNEDenoiserSettings.h
@@ -0,0 +1,46 @@
+// Copyright Epic Games, Inc. All Rights Reserved.
+
+#pragma once
+
+#include "Engine/DeveloperSettingsBackedByCVars.h"
+#include "NNEDenoiserAsset.h"
+#include "NNEDenoiserTemporalAsset.h"
+
+#include "NNEDenoiserSettings.generated.h"
+
+/** An enum to represent denoiser NNE runtime type */
+UENUM()
+enum EDenoiserRuntimeType : uint8
+{
+	CPU,
+	GPU,
+	RDG
+};
+
+/** Settings to select a NNE Denoiser and its runtime */
+UCLASS(Config = Engine, meta = (DisplayName = "NNE Denoiser"))
+class NNEDENOISER_API UNNEDenoiserSettings : public UDeveloperSettingsBackedByCVars
+{
+	GENERATED_BODY()
+
+public:
+	UNNEDenoiserSettings();
+
+	virtual void PostInitProperties() override;
+
+	/** Denoiser asset data used to create a NNE Denoiser */
+	UPROPERTY(Config, EditAnywhere, Category = "NNE Denoiser", meta = (DisplayName = "Denoiser Asset", ToolTip = "Select the denoiser asset"))
+	TSoftObjectPtr<UNNEDenoiserAsset> DenoiserAsset;
+
+	/** Temporal denoiser asset data used to create a NNE Denoiser (Currently not used and therefore "hidden") */
+	TSoftObjectPtr<UNNEDenoiserTemporalAsset> TemporalDenoiserAsset;
+
+private:
+	/** Runtime type used to run the NNE Denoiser model. Backed by the console variable 'NNEDenoiser.Runtime.Type'. */
+	UPROPERTY(Config, EditAnywhere, Category = "NNE Denoiser", meta = (DisplayName = "Runtime Type", ToolTip = "Select a Runtime type", ConsoleVariable = "NNEDenoiser.Runtime.Type"))
+	TEnumAsByte<EDenoiserRuntimeType> RuntimeType;
+
+	/** Runtime name used to run the NNE Denoiser model. Backed by the console variable 'NNEDenoiser.Runtime.Name'. */
+	UPROPERTY(Config, EditAnywhere, Category = "NNE Denoiser", meta = (DisplayName = "Runtime Name Override", ToolTip = "(Optional) Specify the Runtime name", ConsoleVariable = "NNEDenoiser.Runtime.Name"))
+	FString RuntimeName;
+};
diff --git a/Engine/Plugins/NNE/NNEDenoiser/Source/NNEDenoiser/Public/NNEDenoiserTemporalAsset.h b/Engine/Plugins/NNE/NNEDenoiser/Source/NNEDenoiser/Public/NNEDenoiserTemporalAsset.h
new file mode 100644
--- /dev/null
+++ b/Engine/Plugins/NNE/NNEDenoiser/Source/NNEDenoiser/Public/NNEDenoiserTemporalAsset.h
@@ -0,0 +1,34 @@
+// Copyright Epic Games, Inc. All Rights Reserved.
+
+#pragma once
+
+#include "Engine/DataAsset.h"
+#include "Engine/DataTable.h"
+#include "NNEDenoiserTilingConfig.h"
+#include "NNEModelData.h"
+
+#include "NNEDenoiserTemporalAsset.generated.h"
+
+/** Denoiser model data asset */
+UCLASS(BlueprintType)
+class NNEDENOISER_API UNNEDenoiserTemporalAsset : public UDataAsset
+{
+	GENERATED_BODY()
+
+public:
+	/** NNE model data */
+	UPROPERTY(EditAnywhere, BlueprintReadWrite, Category=NNEDenoiser)
+	TSoftObjectPtr<UNNEModelData> ModelData;
+
+	/** Input mapping table */
+	UPROPERTY(EditAnywhere, BlueprintReadWrite, Category=NNEDenoiser, meta = (RequiredAssetDataTags = "RowStructure=/Script/NNEDenoiser.NNEDenoiserTemporalInputMappingData"))
+	TSoftObjectPtr<UDataTable> InputMapping;
+
+	/** Output mapping table */
+	UPROPERTY(EditAnywhere, BlueprintReadWrite, Category=NNEDenoiser, meta = (RequiredAssetDataTags = "RowStructure=/Script/NNEDenoiser.NNEDenoiserTemporalOutputMappingData"))
+	TSoftObjectPtr<UDataTable> OutputMapping;
+
+	/** Tiling configuration */
+	UPROPERTY(EditAnywhere, BlueprintReadWrite, Category=NNEDenoiser)
+	FTilingConfig TilingConfig{};
+};
\ No newline at end of file
diff --git a/Engine/Plugins/NNE/NNEDenoiser/Source/NNEDenoiser/Public/NNEDenoiserTilingConfig.h b/Engine/Plugins/NNE/NNEDenoiser/Source/NNEDenoiser/Public/NNEDenoiserTilingConfig.h
new file mode 100644
--- /dev/null
+++ b/Engine/Plugins/NNE/NNEDenoiser/Source/NNEDenoiser/Public/NNEDenoiserTilingConfig.h
@@ -0,0 +1,28 @@
+// Copyright Epic Games, Inc. All Rights Reserved.
+
+#pragma once
+
+#include "NNEDenoiserTilingConfig.generated.h"
+
+/** Tiling configuration for fixed and dynamic size models */
+USTRUCT(BlueprintType)
+struct FTilingConfig
+{
+	GENERATED_BODY()
+
+	/** Tile size alignment (applies only to dynamic size models) */
+	UPROPERTY(EditAnywhere, BlueprintReadWrite, Category=NNEDenoiser, meta = (DisplayName = "Size Alignment"))
+	int32 Alignment = 1;
+
+	/** Tile overlap */
+	UPROPERTY(EditAnywhere, BlueprintReadWrite, Category=NNEDenoiser)
+	int32 Overlap = 0;
+
+	/** Maximum tile size (applies only to dynamic size models) */
+	UPROPERTY(EditAnywhere, BlueprintReadWrite, Category=NNEDenoiser)
+	int32 MaxSize = 0;
+
+	/** Minimum tile size (applies only to dynamic size models) */
+	UPROPERTY(EditAnywhere, BlueprintReadWrite, Category=NNEDenoiser)
+	int32 MinSize = 1;
+};
\ No newline at end of file
diff --git a/Engine/Plugins/NNE/NNEDenoiser/Source/NNEDenoiserShaders/Internal/NNEDenoiserShadersAutoExposureCS.h b/Engine/Plugins/NNE/NNEDenoiser/Source/NNEDenoiserShaders/Internal/NNEDenoiserShadersAutoExposureCS.h
new file mode 100644
--- /dev/null
+++ b/Engine/Plugins/NNE/NNEDenoiser/Source/NNEDenoiserShaders/Internal/NNEDenoiserShadersAutoExposureCS.h
@@ -0,0 +1,84 @@
+// Copyright Epic Games, Inc. All Rights Reserved.
+
+#pragma once
+
+#include "GlobalShader.h"
+#include "RenderGraphUtils.h"
+#include "ShaderParameterUtils.h"
+
+namespace UE::NNEDenoiserShaders::Internal
+{
+	
+	class FAutoExposureDownsampleConstants
+	{
+	public:
+		static constexpr int32 MAX_BIN_SIZE{16};
+		static constexpr int32 THREAD_GROUP_SIZE{MAX_BIN_SIZE};
+	};
+
+	class NNEDENOISERSHADERS_API FAutoExposureDownsampleCS : public FGlobalShader
+	{
+		DECLARE_GLOBAL_SHADER(FAutoExposureDownsampleCS);
+		SHADER_USE_PARAMETER_STRUCT(FAutoExposureDownsampleCS, FGlobalShader)
+
+	public:
+
+		BEGIN_SHADER_PARAMETER_STRUCT(FParameters, )
+			SHADER_PARAMETER(int32, InputTextureWidth)
+			SHADER_PARAMETER(int32, InputTextureHeight)
+			SHADER_PARAMETER_RDG_TEXTURE(Texture2D, InputTexture)
+			SHADER_PARAMETER(int32, NumBinsW)
+			SHADER_PARAMETER(int32, NumBinsH)
+			SHADER_PARAMETER_RDG_BUFFER_UAV(RWBuffer, OutputBins)
+		END_SHADER_PARAMETER_STRUCT()
+
+		static void ModifyCompilationEnvironment(const FGlobalShaderPermutationParameters& InParameters, FShaderCompilerEnvironment& OutEnvironment);
+		static bool ShouldCompilePermutation(const FGlobalShaderPermutationParameters& Parameters);
+	};
+
+	class FAutoExposureReduceConstants
+	{
+	public:
+		static constexpr int32 THREAD_GROUP_SIZE{32};
+		static constexpr float EPS{1e-8f};
+		static constexpr float KEY{0.18f};
+	};
+
+	class NNEDENOISERSHADERS_API FAutoExposureReduceCS : public FGlobalShader
+	{
+		DECLARE_GLOBAL_SHADER(FAutoExposureReduceCS);
+		SHADER_USE_PARAMETER_STRUCT(FAutoExposureReduceCS, FGlobalShader)
+
+	public:
+
+		BEGIN_SHADER_PARAMETER_STRUCT(FParameters, )
+			SHADER_PARAMETER(int32, InputSize)
+			SHADER_PARAMETER_RDG_BUFFER_SRV(Buffer, InputBins)
+			SHADER_PARAMETER_RDG_BUFFER_UAV(RWBuffer, OutputSums)
+			SHADER_PARAMETER_RDG_BUFFER_UAV(RWBuffer, OutputCounts)
+			SHADER_PARAMETER(int32, NumThreads)
+		END_SHADER_PARAMETER_STRUCT()
+
+		static void ModifyCompilationEnvironment(const FGlobalShaderPermutationParameters& InParameters, FShaderCompilerEnvironment& OutEnvironment);
+		static bool ShouldCompilePermutation(const FGlobalShaderPermutationParameters& Parameters);
+	};
+
+	class NNEDENOISERSHADERS_API FAutoExposureReduceFinalCS : public FGlobalShader
+	{
+		DECLARE_GLOBAL_SHADER(FAutoExposureReduceFinalCS);
+		SHADER_USE_PARAMETER_STRUCT(FAutoExposureReduceFinalCS, FGlobalShader)
+
+	public:
+
+		BEGIN_SHADER_PARAMETER_STRUCT(FParameters, )
+			SHADER_PARAMETER(int32, InputSize)
+			SHADER_PARAMETER_RDG_BUFFER_SRV(Buffer, InputSums)
+			SHADER_PARAMETER_RDG_BUFFER_SRV(Buffer, InputCounts)
+			SHADER_PARAMETER_RDG_BUFFER_UAV(RWBuffer, OutputBuffer)
+		END_SHADER_PARAMETER_STRUCT()
+
+		static void ModifyCompilationEnvironment(const FGlobalShaderPermutationParameters& InParameters, FShaderCompilerEnvironment& OutEnvironment);
+		static bool ShouldCompilePermutation(const FGlobalShaderPermutationParameters& Parameters);
+	};
+
+} // namespace UE::NNEDenoiser::Private
\ No newline at end of file
diff --git a/Engine/Plugins/NNE/NNEDenoiser/Source/NNEDenoiserShaders/Internal/NNEDenoiserShadersDefaultIOProcessCS.h b/Engine/Plugins/NNE/NNEDenoiser/Source/NNEDenoiserShaders/Internal/NNEDenoiserShadersDefaultIOProcessCS.h
new file mode 100644
--- /dev/null
+++ b/Engine/Plugins/NNE/NNEDenoiser/Source/NNEDenoiserShaders/Internal/NNEDenoiserShadersDefaultIOProcessCS.h
@@ -0,0 +1,49 @@
+// Copyright Epic Games, Inc. All Rights Reserved.
+
+#pragma once
+
+#include "GlobalShader.h"
+#include "RenderGraphUtils.h"
+#include "ShaderParameterUtils.h"
+
+namespace UE::NNEDenoiserShaders::Internal
+{
+
+	enum class EDefaultIOProcessInputKind : uint8
+	{
+		Color = 0,
+		Albedo,
+		Normal,
+		Flow,
+		Output,
+		MAX
+	};
+
+	class FDefaultIOProcessConstants
+	{
+	public:
+		static constexpr int32 THREAD_GROUP_SIZE{16};
+	};
+	
+	class NNEDENOISERSHADERS_API FDefaultIOProcessCS : public FGlobalShader
+	{
+		DECLARE_GLOBAL_SHADER(FDefaultIOProcessCS);
+		SHADER_USE_PARAMETER_STRUCT(FDefaultIOProcessCS, FGlobalShader)
+
+		class FDefaultIOProcessInputKind : SHADER_PERMUTATION_ENUM_CLASS("INPUT_KIND_INDEX", EDefaultIOProcessInputKind);
+		using FPermutationDomain = TShaderPermutationDomain<FDefaultIOProcessInputKind>;
+
+	public:
+
+		BEGIN_SHADER_PARAMETER_STRUCT(FParameters, )
+			SHADER_PARAMETER(int32, Width)
+			SHADER_PARAMETER(int32, Height)
+			SHADER_PARAMETER_RDG_TEXTURE(Texture2D, InputTexture)
+			SHADER_PARAMETER_RDG_TEXTURE_UAV(RWTexture2D, OutputTexture)
+		END_SHADER_PARAMETER_STRUCT()
+
+		static void ModifyCompilationEnvironment(const FGlobalShaderPermutationParameters& InParameters, FShaderCompilerEnvironment& OutEnvironment);
+		static bool ShouldCompilePermutation(const FGlobalShaderPermutationParameters& Parameters);
+	};
+
+} // namespace UE::NNEDenoiser::Private
\ No newline at end of file
diff --git a/Engine/Plugins/NNE/NNEDenoiser/Source/NNEDenoiserShaders/Internal/NNEDenoiserShadersMappedCopyCS.h b/Engine/Plugins/NNE/NNEDenoiser/Source/NNEDenoiserShaders/Internal/NNEDenoiserShadersMappedCopyCS.h
new file mode 100644
--- /dev/null
+++ b/Engine/Plugins/NNE/NNEDenoiser/Source/NNEDenoiserShaders/Internal/NNEDenoiserShadersMappedCopyCS.h
@@ -0,0 +1,80 @@
+// Copyright Epic Games, Inc. All Rights Reserved.
+
+#pragma once
+
+#include "GlobalShader.h"
+#include "RenderGraphUtils.h"
+#include "ShaderParameterUtils.h"
+
+namespace UE::NNEDenoiserShaders::Internal
+{
+
+	// Maps to ENNETensorDataType
+	enum class EDataType : uint8
+	{
+		None = 0,
+		// Char,
+		// Boolean,
+		Half = 3,
+		Float,
+		MAX
+	};
+
+	class FMappedCopyConstants
+	{
+	public:
+		static constexpr int32 THREAD_GROUP_SIZE{16};
+		static constexpr int32 MAX_NUM_MAPPED_CHANNELS{4};
+	};
+
+	// Note: The shader also supports Texture/Texture and Buffer/Buffer mapped copy, just add the Shader class.
+	
+	class NNEDENOISERSHADERS_API FTextureBufferMappedCopyCS : public FGlobalShader
+	{
+		DECLARE_GLOBAL_SHADER(FTextureBufferMappedCopyCS);
+		SHADER_USE_PARAMETER_STRUCT(FTextureBufferMappedCopyCS, FGlobalShader)
+
+		class FInputDataType : SHADER_PERMUTATION_ENUM_CLASS("INPUT_DATA_TYPE_INDEX", EDataType);
+		class FOutputDataType : SHADER_PERMUTATION_ENUM_CLASS("OUTPUT_DATA_TYPE_INDEX", EDataType);
+		class FNumMappedChannels : SHADER_PERMUTATION_RANGE_INT("NUM_MAPPED_CHANNELS", 0, FMappedCopyConstants::MAX_NUM_MAPPED_CHANNELS + 1);
+		using FPermutationDomain = TShaderPermutationDomain<FInputDataType, FOutputDataType, FNumMappedChannels>;
+
+	public:
+
+		BEGIN_SHADER_PARAMETER_STRUCT(FParameters, )
+			SHADER_PARAMETER(int32, Width)
+			SHADER_PARAMETER(int32, Height)
+			SHADER_PARAMETER_RDG_TEXTURE(Texture2D, InputTexture)
+			SHADER_PARAMETER_RDG_BUFFER_UAV(RWBuffer, OutputBuffer)
+			SHADER_PARAMETER_ARRAY(FIntVector4, OutputChannel_InputChannel_Unused_Unused, [FMappedCopyConstants::MAX_NUM_MAPPED_CHANNELS])
+		END_SHADER_PARAMETER_STRUCT()
+
+		static void ModifyCompilationEnvironment(const FGlobalShaderPermutationParameters& InParameters, FShaderCompilerEnvironment& OutEnvironment);
+		static bool ShouldCompilePermutation(const FGlobalShaderPermutationParameters& Parameters);
+	};
+
+	class NNEDENOISERSHADERS_API FBufferTextureMappedCopyCS : public FGlobalShader
+	{
+		DECLARE_GLOBAL_SHADER(FBufferTextureMappedCopyCS);
+		SHADER_USE_PARAMETER_STRUCT(FBufferTextureMappedCopyCS, FGlobalShader)
+
+		class FInputDataType : SHADER_PERMUTATION_ENUM_CLASS("INPUT_DATA_TYPE_INDEX", EDataType);
+		class FOutputDataType : SHADER_PERMUTATION_ENUM_CLASS("OUTPUT_DATA_TYPE_INDEX", EDataType);
+		class FNumMappedChannels : SHADER_PERMUTATION_RANGE_INT("NUM_MAPPED_CHANNELS", 0, FMappedCopyConstants::MAX_NUM_MAPPED_CHANNELS + 1);
+		using FPermutationDomain = TShaderPermutationDomain<FInputDataType, FOutputDataType, FNumMappedChannels>;
+
+	public:
+
+		BEGIN_SHADER_PARAMETER_STRUCT(FParameters, )
+			SHADER_PARAMETER(int32, Width)
+			SHADER_PARAMETER(int32, Height)
+			SHADER_PARAMETER_RDG_BUFFER_SRV(Buffer, InputBuffer)
+			SHADER_PARAMETER_RDG_TEXTURE_UAV(RWTexture2D, OutputTexture)
+			SHADER_PARAMETER_ARRAY(FIntVector4, OutputChannel_InputChannel_Unused_Unused, [FMappedCopyConstants::MAX_NUM_MAPPED_CHANNELS])
+		END_SHADER_PARAMETER_STRUCT()
+
+		static void ModifyCompilationEnvironment(const FGlobalShaderPermutationParameters& InParameters, FShaderCompilerEnvironment& OutEnvironment);
+		static bool ShouldCompilePermutation(const FGlobalShaderPermutationParameters& Parameters);
+	};
+
+} // namespace UE::NNEDenoiser::Private
\ No newline at end of file
diff --git a/Engine/Plugins/NNE/NNEDenoiser/Source/NNEDenoiserShaders/Internal/NNEDenoiserShadersTransferFunctionOidnCS.h b/Engine/Plugins/NNE/NNEDenoiser/Source/NNEDenoiserShaders/Internal/NNEDenoiserShadersTransferFunctionOidnCS.h
new file mode 100644
--- /dev/null
+++ b/Engine/Plugins/NNE/NNEDenoiser/Source/NNEDenoiserShaders/Internal/NNEDenoiserShadersTransferFunctionOidnCS.h
@@ -0,0 +1,49 @@
+// Copyright Epic Games, Inc. All Rights Reserved.
+
+#pragma once
+
+#include "GlobalShader.h"
+#include "RenderGraphUtils.h"
+#include "ShaderParameterUtils.h"
+
+namespace UE::NNEDenoiserShaders::Internal
+{
+
+	enum class ETransferFunctionOidnMode : uint8
+	{
+		Forwward = 0,
+		Inverse,
+		MAX
+	};
+
+	class FTransferFunctionOidnConstants
+	{
+	public:
+		static constexpr int32 THREAD_GROUP_SIZE{16};
+	};
+	
+	class NNEDENOISERSHADERS_API FTransferFunctionOidnCS : public FGlobalShader
+	{
+		DECLARE_GLOBAL_SHADER(FTransferFunctionOidnCS);
+		SHADER_USE_PARAMETER_STRUCT(FTransferFunctionOidnCS, FGlobalShader)
+
+		class FTransferFunctionOidnMode : SHADER_PERMUTATION_ENUM_CLASS("MODE", ETransferFunctionOidnMode);
+		using FPermutationDomain = TShaderPermutationDomain<FTransferFunctionOidnMode>;
+
+	public:
+
+		BEGIN_SHADER_PARAMETER_STRUCT(FParameters, )
+			SHADER_PARAMETER(int32, Width)
+			SHADER_PARAMETER(int32, Height)
+			SHADER_PARAMETER_RDG_TEXTURE(Texture2D, InputTexture)
+			SHADER_PARAMETER_RDG_BUFFER_SRV(Buffer, InputScaleBuffer)
+			SHADER_PARAMETER(float, NormScale)
+			SHADER_PARAMETER(float, InvNormScale)
+			SHADER_PARAMETER_RDG_TEXTURE_UAV(RWTexture2D, OutputTexture)
+		END_SHADER_PARAMETER_STRUCT()
+
+		static void ModifyCompilationEnvironment(const FGlobalShaderPermutationParameters& InParameters, FShaderCompilerEnvironment& OutEnvironment);
+		static bool ShouldCompilePermutation(const FGlobalShaderPermutationParameters& Parameters);
+	};
+
+} // namespace UE::NNEDenoiser::Private
\ No newline at end of file
diff --git a/Engine/Plugins/NNE/NNEDenoiser/Source/NNEDenoiserShaders/NNEDenoiserShaders.Build.cs b/Engine/Plugins/NNE/NNEDenoiser/Source/NNEDenoiserShaders/NNEDenoiserShaders.Build.cs
new file mode 100644
--- /dev/null
+++ b/Engine/Plugins/NNE/NNEDenoiser/Source/NNEDenoiserShaders/NNEDenoiserShaders.Build.cs
@@ -0,0 +1,26 @@
+// Copyright Epic Games, Inc. All Rights Reserved.
+
+using UnrealBuildTool;
+
+public class NNEDenoiserShaders : ModuleRules
+{
+	public NNEDenoiserShaders(ReadOnlyTargetRules Target) : base(Target)
+	{
+		PCHUsage = PCHUsageMode.UseExplicitOrSharedPCHs;
+
+		PublicDependencyModuleNames.AddRange(
+			new string[]
+			{
+				"Core",
+				"RenderCore"
+			});
+			
+		PrivateDependencyModuleNames.AddRange(
+			new string[]
+			{
+				"Projects",
+				"Renderer",
+				"RHI"
+			});
+	}
+}
diff --git a/Engine/Plugins/NNE/NNEDenoiser/Source/NNEDenoiserShaders/Private/NNEDenoiserShadersAutoExposureCS.cpp b/Engine/Plugins/NNE/NNEDenoiser/Source/NNEDenoiserShaders/Private/NNEDenoiserShadersAutoExposureCS.cpp
new file mode 100644
--- /dev/null
+++ b/Engine/Plugins/NNE/NNEDenoiser/Source/NNEDenoiserShaders/Private/NNEDenoiserShadersAutoExposureCS.cpp
@@ -0,0 +1,48 @@
+// Copyright Epic Games, Inc. All Rights Reserved.
+
+#include "NNEDenoiserShadersAutoExposureCS.h"
+
+namespace UE::NNEDenoiserShaders::Internal
+{
+
+	void FAutoExposureDownsampleCS::ModifyCompilationEnvironment(const FGlobalShaderPermutationParameters& InParameters, FShaderCompilerEnvironment& OutEnvironment)
+	{
+		FGlobalShader::ModifyCompilationEnvironment(InParameters, OutEnvironment);
+		OutEnvironment.SetDefine(TEXT("THREAD_GROUP_SIZE"), FAutoExposureDownsampleConstants::THREAD_GROUP_SIZE);
+		OutEnvironment.SetDefine(TEXT("MAX_FLT"), MAX_FLT);
+	}
+
+	void FAutoExposureReduceCS::ModifyCompilationEnvironment(const FGlobalShaderPermutationParameters& InParameters, FShaderCompilerEnvironment& OutEnvironment)
+	{
+		FGlobalShader::ModifyCompilationEnvironment(InParameters, OutEnvironment);
+		OutEnvironment.SetDefine(TEXT("THREAD_GROUP_SIZE"), FAutoExposureReduceConstants::THREAD_GROUP_SIZE);
+		OutEnvironment.SetDefine(TEXT("EPS"), FAutoExposureReduceConstants::EPS);
+	}
+
+	void FAutoExposureReduceFinalCS::ModifyCompilationEnvironment(const FGlobalShaderPermutationParameters& InParameters, FShaderCompilerEnvironment& OutEnvironment)
+	{
+		FGlobalShader::ModifyCompilationEnvironment(InParameters, OutEnvironment);
+		OutEnvironment.SetDefine(TEXT("THREAD_GROUP_SIZE"), FAutoExposureReduceConstants::THREAD_GROUP_SIZE);
+		OutEnvironment.SetDefine(TEXT("KEY"), FAutoExposureReduceConstants::KEY);
+	}
+
+	bool FAutoExposureDownsampleCS::ShouldCompilePermutation(const FGlobalShaderPermutationParameters& Parameters)
+	{
+		return ShouldCompileRayTracingShadersForProject(Parameters.Platform);
+	}
+
+	bool FAutoExposureReduceCS::ShouldCompilePermutation(const FGlobalShaderPermutationParameters& Parameters)
+	{
+		return ShouldCompileRayTracingShadersForProject(Parameters.Platform);
+	}
+
+	bool FAutoExposureReduceFinalCS::ShouldCompilePermutation(const FGlobalShaderPermutationParameters& Parameters)
+	{
+		return ShouldCompileRayTracingShadersForProject(Parameters.Platform);
+	}
+
+	IMPLEMENT_GLOBAL_SHADER(FAutoExposureDownsampleCS, "/NNEDenoiserShaders/NNEDenoiserShadersAutoExposureDownsample.usf", "Downsample", SF_Compute);
+	IMPLEMENT_GLOBAL_SHADER(FAutoExposureReduceCS, "/NNEDenoiserShaders/NNEDenoiserShadersAutoExposureReduce.usf", "Reduce", SF_Compute);
+	IMPLEMENT_GLOBAL_SHADER(FAutoExposureReduceFinalCS, "/NNEDenoiserShaders/NNEDenoiserShadersAutoExposureReduceFinal.usf", "ReduceFinal", SF_Compute);
+
+} // UE::NNEDenoiser::Private
\ No newline at end of file
diff --git a/Engine/Plugins/NNE/NNEDenoiser/Source/NNEDenoiserShaders/Private/NNEDenoiserShadersDefaultIOProcessCS.cpp b/Engine/Plugins/NNE/NNEDenoiser/Source/NNEDenoiserShaders/Private/NNEDenoiserShadersDefaultIOProcessCS.cpp
new file mode 100644
--- /dev/null
+++ b/Engine/Plugins/NNE/NNEDenoiser/Source/NNEDenoiserShaders/Private/NNEDenoiserShadersDefaultIOProcessCS.cpp
@@ -0,0 +1,21 @@
+// Copyright Epic Games, Inc. All Rights Reserved.
+
+#include "NNEDenoiserShadersDefaultIOProcessCS.h"
+
+namespace UE::NNEDenoiserShaders::Internal
+{
+	void FDefaultIOProcessCS::ModifyCompilationEnvironment(const FGlobalShaderPermutationParameters& InParameters, FShaderCompilerEnvironment& OutEnvironment)
+	{
+		FGlobalShader::ModifyCompilationEnvironment(InParameters, OutEnvironment);
+		OutEnvironment.SetDefine(TEXT("THREAD_GROUP_SIZE"), FDefaultIOProcessConstants::THREAD_GROUP_SIZE);
+		OutEnvironment.SetDefine(TEXT("MAX_FLT"), MAX_FLT);
+	}
+
+	bool FDefaultIOProcessCS::ShouldCompilePermutation(const FGlobalShaderPermutationParameters& Parameters)
+	{
+		return ShouldCompileRayTracingShadersForProject(Parameters.Platform);
+	}
+
+	IMPLEMENT_GLOBAL_SHADER(FDefaultIOProcessCS, "/NNEDenoiserShaders/NNEDenoiserShadersDefaultIOProcess.usf", "IOProcess", SF_Compute);
+
+} // UE::NNEDenoiser::Private
\ No newline at end of file
diff --git a/Engine/Plugins/NNE/NNEDenoiser/Source/NNEDenoiserShaders/Private/NNEDenoiserShadersLog.h b/Engine/Plugins/NNE/NNEDenoiser/Source/NNEDenoiserShaders/Private/NNEDenoiserShadersLog.h
new file mode 100644
--- /dev/null
+++ b/Engine/Plugins/NNE/NNEDenoiser/Source/NNEDenoiserShaders/Private/NNEDenoiserShadersLog.h
@@ -0,0 +1,8 @@
+// Copyright Epic Games, Inc. All Rights Reserved.
+
+#pragma once
+
+#include "CoreMinimal.h"
+#include "Logging/LogMacros.h"
+
+DECLARE_LOG_CATEGORY_EXTERN(LogNNEDenoiserShaders, Log, All);
\ No newline at end of file
diff --git a/Engine/Plugins/NNE/NNEDenoiser/Source/NNEDenoiserShaders/Private/NNEDenoiserShadersMappedCopyCS.cpp b/Engine/Plugins/NNE/NNEDenoiser/Source/NNEDenoiserShaders/Private/NNEDenoiserShadersMappedCopyCS.cpp
new file mode 100644
--- /dev/null
+++ b/Engine/Plugins/NNE/NNEDenoiser/Source/NNEDenoiserShaders/Private/NNEDenoiserShadersMappedCopyCS.cpp
@@ -0,0 +1,64 @@
+// Copyright Epic Games, Inc. All Rights Reserved.
+
+#include "NNEDenoiserShadersMappedCopyCS.h"
+
+namespace UE::NNEDenoiserShaders::Internal
+{
+	void CommonModifyCompilationEnvironment(const FGlobalShaderPermutationParameters& InParameters, FShaderCompilerEnvironment& OutEnvironment)
+	{
+		FGlobalShader::ModifyCompilationEnvironment(InParameters, OutEnvironment);
+		OutEnvironment.SetDefine(TEXT("THREAD_GROUP_SIZE"), FMappedCopyConstants::THREAD_GROUP_SIZE);
+		OutEnvironment.SetDefine(TEXT("MAX_NUM_MAPPED_CHANNELS"), FMappedCopyConstants::MAX_NUM_MAPPED_CHANNELS);
+	}
+
+	template<class GlobalShaderType>
+	bool CommonShouldCompilePermutation(const FGlobalShaderPermutationParameters& Parameters)
+	{
+		if (!ShouldCompileRayTracingShadersForProject(Parameters.Platform))
+		{
+			return false;
+		}
+
+		typename GlobalShaderType::FPermutationDomain PermutationVector(Parameters.PermutationId);
+		if (PermutationVector.template Get<typename GlobalShaderType::FInputDataType>() != EDataType::Half &&
+			PermutationVector.template Get<typename GlobalShaderType::FInputDataType>() != EDataType::Float)
+		{
+			return false;
+		}
+		if (PermutationVector.template Get<typename GlobalShaderType::FOutputDataType>() != EDataType::Half &&
+			PermutationVector.template Get<typename GlobalShaderType::FOutputDataType>() != EDataType::Float)
+		{
+			return false;
+		}
+
+		return true;
+	}
+
+	void FTextureBufferMappedCopyCS::ModifyCompilationEnvironment(const FGlobalShaderPermutationParameters& InParameters, FShaderCompilerEnvironment& OutEnvironment)
+	{
+		CommonModifyCompilationEnvironment(InParameters, OutEnvironment);
+		OutEnvironment.SetDefine(TEXT("INTRINSIC_INPUT_TYPE"), 0);
+		OutEnvironment.SetDefine(TEXT("INTRINSIC_OUTPUT_TYPE"), 1);
+	}
+
+	void FBufferTextureMappedCopyCS::ModifyCompilationEnvironment(const FGlobalShaderPermutationParameters& InParameters, FShaderCompilerEnvironment& OutEnvironment)
+	{
+		CommonModifyCompilationEnvironment(InParameters, OutEnvironment);
+		OutEnvironment.SetDefine(TEXT("INTRINSIC_INPUT_TYPE"), 1);
+		OutEnvironment.SetDefine(TEXT("INTRINSIC_OUTPUT_TYPE"), 0);
+	}
+
+	bool FTextureBufferMappedCopyCS::ShouldCompilePermutation(const FGlobalShaderPermutationParameters& Parameters)
+	{
+		return CommonShouldCompilePermutation<FTextureBufferMappedCopyCS>(Parameters);
+	}
+
+	bool FBufferTextureMappedCopyCS::ShouldCompilePermutation(const FGlobalShaderPermutationParameters& Parameters)
+	{
+		return CommonShouldCompilePermutation<FBufferTextureMappedCopyCS>(Parameters);
+	}
+
+	IMPLEMENT_GLOBAL_SHADER(FTextureBufferMappedCopyCS, "/NNEDenoiserShaders/NNEDenoiserShadersMappedCopy.usf", "MappedCopy", SF_Compute);
+	IMPLEMENT_GLOBAL_SHADER(FBufferTextureMappedCopyCS, "/NNEDenoiserShaders/NNEDenoiserShadersMappedCopy.usf", "MappedCopy", SF_Compute);
+
+} // UE::NNEDenoiser::Private
\ No newline at end of file
diff --git a/Engine/Plugins/NNE/NNEDenoiser/Source/NNEDenoiserShaders/Private/NNEDenoiserShadersModule.cpp b/Engine/Plugins/NNE/NNEDenoiser/Source/NNEDenoiserShaders/Private/NNEDenoiserShadersModule.cpp
new file mode 100644
--- /dev/null
+++ b/Engine/Plugins/NNE/NNEDenoiser/Source/NNEDenoiserShaders/Private/NNEDenoiserShadersModule.cpp
@@ -0,0 +1,37 @@
+// Copyright Epic Games, Inc. All Rights Reserved.
+
+#include "CoreMinimal.h"
+#include "Interfaces/IPluginManager.h"
+#include "Misc/Paths.h"
+#include "Modules/ModuleManager.h"
+#include "ShaderCore.h"
+#include "NNEDenoiserShadersLog.h"
+
+DEFINE_LOG_CATEGORY(LogNNEDenoiserShaders);
+
+class FNNEDenoiserShadersModule : public IModuleInterface
+{
+public:
+	// Begin IModuleInterface
+	virtual void StartupModule() override
+	{
+		const TSharedPtr<IPlugin> Plugin = IPluginManager::Get().FindPlugin(TEXT("NNEDenoiser"));
+		if (Plugin.IsValid())
+		{
+			const FString ShadersDir = FPaths::Combine(Plugin->GetBaseDir(), TEXT("Shaders"));
+
+			AddShaderSourceDirectoryMapping(TEXT("/NNEDenoiserShaders"), ShadersDir);
+		}
+		else
+		{
+			UE_LOG(LogNNEDenoiserShaders, Warning, TEXT("Shaders directory not added. Failed to find NNEDenoiser plugin"));
+		}
+	}
+
+	virtual void ShutdownModule() override
+	{
+
+	}
+};
+
+IMPLEMENT_MODULE(FNNEDenoiserShadersModule, NNEDenoiserShaders);
diff --git a/Engine/Plugins/NNE/NNEDenoiser/Source/NNEDenoiserShaders/Private/NNEDenoiserShadersTransferFunctionOidnCS.cpp b/Engine/Plugins/NNE/NNEDenoiser/Source/NNEDenoiserShaders/Private/NNEDenoiserShadersTransferFunctionOidnCS.cpp
new file mode 100644
--- /dev/null
+++ b/Engine/Plugins/NNE/NNEDenoiser/Source/NNEDenoiserShaders/Private/NNEDenoiserShadersTransferFunctionOidnCS.cpp
@@ -0,0 +1,21 @@
+// Copyright Epic Games, Inc. All Rights Reserved.
+
+#include "NNEDenoiserShadersTransferFunctionOidnCS.h"
+
+namespace UE::NNEDenoiserShaders::Internal
+{
+	void FTransferFunctionOidnCS::ModifyCompilationEnvironment(const FGlobalShaderPermutationParameters& InParameters, FShaderCompilerEnvironment& OutEnvironment)
+	{
+		FGlobalShader::ModifyCompilationEnvironment(InParameters, OutEnvironment);
+		OutEnvironment.SetDefine(TEXT("THREAD_GROUP_SIZE"), FTransferFunctionOidnConstants::THREAD_GROUP_SIZE);
+		OutEnvironment.SetDefine(TEXT("MAX_FLT"), MAX_FLT);
+	}
+
+	bool FTransferFunctionOidnCS::ShouldCompilePermutation(const FGlobalShaderPermutationParameters& Parameters)
+	{
+		return ShouldCompileRayTracingShadersForProject(Parameters.Platform);
+	}
+
+	IMPLEMENT_GLOBAL_SHADER(FTransferFunctionOidnCS, "/NNEDenoiserShaders/NNEDenoiserShadersTransferFunctionOidn.usf", "PreOrPostprocess", SF_Compute);
+
+} // UE::NNEDenoiser::Private
\ No newline at end of file
