[NNEDenoiser] Added plugin config to add /NNEDenoiser path to cook directories.

#tests Packaged game and verified NNEDenosier assets are included

#rnx

diff --git a/Engine/Plugins/NNE/NNEDenoiser/Config/Game.ini b/Engine/Plugins/NNE/NNEDenoiser/Config/Game.ini
new file mode 100644
--- /dev/null
+++ b/Engine/Plugins/NNE/NNEDenoiser/Config/Game.ini
@@ -0,0 +1,2 @@
+[/Script/UnrealEd.ProjectPackagingSettings]
++DirectoriesToAlwaysCook=(Path="/NNEDenoiser")
\ No newline at end of file
