[NNEDenoiser] Improved resource management, specifically only create denoiser and load assets if Path Tracer is enabled in one of the view families.

Problem:
NNEDenoiser utilizes resources even when Path Tracer is disabled. Currently the lazy synchronous loading causes warning in some Game code.

Solution:
Properly check whether Path Tracer is enabled and only load assets and create denoiser if actually used.

#tests Editor, Win64, verified resources freed when disabling Path Tracer

#rnx

diff --git a/Engine/Plugins/NNE/NNEDenoiser/Source/NNEDenoiser/Private/NNEDenoise.cpp b/Engine/Plugins/NNE/NNEDenoiser/Source/NNEDenoiser/Private/NNEDenoise.cpp
--- a/Engine/Plugins/NNE/NNEDenoiser/Source/NNEDenoiser/Private/NNEDenoise.cpp
+++ b/Engine/Plugins/NNE/NNEDenoiser/Source/NNEDenoiser/Private/NNEDenoise.cpp
@@ -9,6 +9,7 @@ DEFINE_LOG_CATEGORY(LogNNEDenoiser);
 
 void FNNEDenoiserModule::StartupModule()
 {
+#if RHI_RAYTRACING
 	// During cook, we are not allowed to call IsRayTracingAllowed() later, therefore we assume Ray Tracing is not available and do not register the View Extension.
 	if (IsRunningCookCommandlet())
 	{
@@ -27,6 +28,7 @@ void FNNEDenoiserModule::StartupModule()
 			UE_LOG(LogNNEDenoiser, Log, TEXT("Ray Tracing is not enabled, therefore NNEDenoiser is not registered!"));
 		}
 	});
+#endif // RHI_RAYTRACING
 }
 
 void FNNEDenoiserModule::ShutdownModule()
diff --git a/Engine/Plugins/NNE/NNEDenoiser/Source/NNEDenoiser/Private/NNEDenoiserViewExtension.cpp b/Engine/Plugins/NNE/NNEDenoiser/Source/NNEDenoiser/Private/NNEDenoiserViewExtension.cpp
--- a/Engine/Plugins/NNE/NNEDenoiser/Source/NNEDenoiser/Private/NNEDenoiserViewExtension.cpp
+++ b/Engine/Plugins/NNE/NNEDenoiser/Source/NNEDenoiser/Private/NNEDenoiserViewExtension.cpp
@@ -19,6 +19,7 @@
 #include "NNEDenoiserUtils.h"
 #include "NNEModelData.h"
 #include "PathTracingDenoiser.h"
+#include "SceneView.h"
 
 namespace UE::NNEDenoiser::Private
 {
@@ -110,6 +111,16 @@ FString GetDenoiserTemporalAssetNameFromCVarAndSettings(const UNNEDenoiserSettin
 	return FString();
 }
 
+bool IsPathTracingEnabled(FSceneViewFamily& InViewFamily)
+{
+	if (!IsRayTracingEnabled())
+	{
+		return false;
+	}
+
+	return !!InViewFamily.EngineShowFlags.PathTracing;
+}
+
 TUniquePtr<FGenericDenoiser> CreateNNEDenoiser(
 	UNNEModelData& ModelData,
 	EDenoiserRuntimeType RuntimeType,
@@ -319,9 +330,7 @@ TUniquePtr<FGenericDenoiser> CreateNNEDenoiserFromAsset(const FString& AssetName
 
 FViewExtension::FViewExtension(const FAutoRegister& AutoRegister) : FSceneViewExtensionBase(AutoRegister)
 {
-#if WITH_EDITOR
-	GetMutableDefault<UNNEDenoiserSettings>()->OnSettingChanged().AddRaw(this, &FViewExtension::OnDenoiserSettingsChanged);
-#endif
+
 }
 
 FViewExtension::~FViewExtension()
@@ -338,38 +347,11 @@ FViewExtension::~FViewExtension()
 
 void FViewExtension::BeginRenderViewFamily(FSceneViewFamily& InViewFamily)
 {
-	ApplySettings(GetDefault<UNNEDenoiserSettings>());
-}
+	const UNNEDenoiserSettings* Settings = GetDefault<UNNEDenoiserSettings>();
 
-void FViewExtension::PreRenderViewFamily_RenderThread(FRDGBuilder& GraphBuilder, FSceneViewFamily& InViewFamily)
-{
-	if (!bDenoiserEnabled)
-	{
-		UnregisterDenoiser(TEXT("NNEDenoiser"));
-	}
-	else
-	{
-		if (DenoiserToSwap.IsValid() || SpatialTemporalDenoiserToSwap.IsValid())
-		{
-			UnregisterDenoiser(TEXT("NNEDenoiser"));
-		}
-
-		if (DenoiserToSwap.IsValid())
-		{
-			RegisterSpatialDenoiser(MoveTemp(DenoiserToSwap),TEXT("NNEDenoiser"));
-		}
-
-		if (SpatialTemporalDenoiserToSwap.IsValid())
-		{
-			RegisterSpatialTemporalDenoiser(MoveTemp(SpatialTemporalDenoiserToSwap), TEXT("NNEDenoiser"));
-		}
-	}
-}
-
-void FViewExtension::ApplySettings(const UNNEDenoiserSettings* Settings)
-{
 	// Check for changes
 	bool bNeedsUpdate = false;
+	bNeedsUpdate |= IsPathTracingEnabled(InViewFamily) != bPathTracingEnabled;
 	bNeedsUpdate |= CVarNNEDenoiser.GetValueOnGameThread() != bDenoiserEnabled;
 	bNeedsUpdate |= GetDenoiserRuntimeTypeFromCVar() != RuntimeType;
 	bNeedsUpdate |= CVarNNEDenoiserRuntimeName.GetValueOnGameThread() != RuntimeName;
@@ -381,11 +363,13 @@ void FViewExtension::ApplySettings(const UNNEDenoiserSettings* Settings)
 		return;
 	}
 
+	bPathTracingEnabled = IsPathTracingEnabled(InViewFamily);
 	bDenoiserEnabled = CVarNNEDenoiser.GetValueOnGameThread();
 
-	UE_LOG(LogNNEDenoiser, Log, TEXT("ApplySettings: bDenoiserEnabled %d"), bDenoiserEnabled);
+	UE_LOG(LogNNEDenoiser, Log, TEXT("bPathTracingEnabled %d"), bPathTracingEnabled);
+	UE_LOG(LogNNEDenoiser, Log, TEXT("bDenoiserEnabled %d"), bDenoiserEnabled);
 
-	if (!bDenoiserEnabled)
+	if (!bPathTracingEnabled || !bDenoiserEnabled)
 	{
 		return;
 	}
@@ -422,13 +406,29 @@ void FViewExtension::ApplySettings(const UNNEDenoiserSettings* Settings)
 	}
 }
 
-#if WITH_EDITOR
-void FViewExtension::OnDenoiserSettingsChanged(UObject* InObject, struct FPropertyChangedEvent& InPropertyChangedEvent)
+void FViewExtension::PreRenderViewFamily_RenderThread(FRDGBuilder& GraphBuilder, FSceneViewFamily& InViewFamily)
 {
-	UE_LOG(LogNNEDenoiser, Log, TEXT("Settings %s changed: %s"), *InObject->GetName(), *InPropertyChangedEvent.GetPropertyName().ToString());
+	if (!bPathTracingEnabled || !bDenoiserEnabled)
+	{
+		UnregisterDenoiser(TEXT("NNEDenoiser"));
+	}
+	else
+	{
+		if (DenoiserToSwap.IsValid() || SpatialTemporalDenoiserToSwap.IsValid())
+		{
+			UnregisterDenoiser(TEXT("NNEDenoiser"));
+		}
+
+		if (DenoiserToSwap.IsValid())
+		{
+			RegisterSpatialDenoiser(MoveTemp(DenoiserToSwap),TEXT("NNEDenoiser"));
+		}
 
-	ApplySettings(CastChecked<UNNEDenoiserSettings>(InObject));
+		if (SpatialTemporalDenoiserToSwap.IsValid())
+		{
+			RegisterSpatialTemporalDenoiser(MoveTemp(SpatialTemporalDenoiserToSwap), TEXT("NNEDenoiser"));
+		}
+	}
 }
-#endif
 
 } // namespace UE::NNEDenoiser::Private
\ No newline at end of file
diff --git a/Engine/Plugins/NNE/NNEDenoiser/Source/NNEDenoiser/Private/NNEDenoiserViewExtension.h b/Engine/Plugins/NNE/NNEDenoiser/Source/NNEDenoiser/Private/NNEDenoiserViewExtension.h
--- a/Engine/Plugins/NNE/NNEDenoiser/Source/NNEDenoiser/Private/NNEDenoiserViewExtension.h
+++ b/Engine/Plugins/NNE/NNEDenoiser/Source/NNEDenoiser/Private/NNEDenoiserViewExtension.h
@@ -28,16 +28,11 @@ class FViewExtension final : public FSceneViewExtensionBase
 		void PreRenderViewFamily_RenderThread(FRDGBuilder& GraphBuilder, FSceneViewFamily& InViewFamily) override;
 
 	private:
-		void ApplySettings(const UNNEDenoiserSettings* Settings);
-
-#if WITH_EDITOR
-		void OnDenoiserSettingsChanged(UObject* InObject, struct FPropertyChangedEvent& InPropertyChangedEvent);
-#endif
-
 		TUniquePtr<UE::Renderer::Private::IPathTracingDenoiser> DenoiserToSwap;
 		TUniquePtr<UE::Renderer::Private::IPathTracingSpatialTemporalDenoiser> SpatialTemporalDenoiserToSwap;
 
 		// Cached settings and CVar values
+		bool bPathTracingEnabled = false;
 		bool bDenoiserEnabled = true;
 		EDenoiserRuntimeType RuntimeType;
 		FString RuntimeName;
