[NNEDenoiser] Remove unnecessary logs (potentially spamming every frame).

Problem:
Log output is in some cases printed in every frame.

Solution:
Remove, since its not really necessary.

Tests:
Open Editor on Windows and run denoiser.

#tests Editor, Win64 run denoiser.

#rnx

diff --git a/Engine/Plugins/NNE/NNEDenoiser/Source/NNEDenoiser/Private/NNEDenoiserViewExtension.cpp b/Engine/Plugins/NNE/NNEDenoiser/Source/NNEDenoiser/Private/NNEDenoiserViewExtension.cpp
--- a/Engine/Plugins/NNE/NNEDenoiser/Source/NNEDenoiser/Private/NNEDenoiserViewExtension.cpp
+++ b/Engine/Plugins/NNE/NNEDenoiser/Source/NNEDenoiser/Private/NNEDenoiserViewExtension.cpp
@@ -366,9 +366,6 @@ void FViewExtension::BeginRenderViewFamily(FSceneViewFamily& InViewFamily)
 	bPathTracingEnabled = IsPathTracingEnabled(InViewFamily);
 	bDenoiserEnabled = CVarNNEDenoiser.GetValueOnGameThread();
 
-	UE_LOG(LogNNEDenoiser, Log, TEXT("bPathTracingEnabled %d"), bPathTracingEnabled);
-	UE_LOG(LogNNEDenoiser, Log, TEXT("bDenoiserEnabled %d"), bDenoiserEnabled);
-
 	if (!bPathTracingEnabled || !bDenoiserEnabled)
 	{
 		return;
