[NNEDenoiser] Added maximum tile size override setting to allow globally override the maximum tile size given by the currently used denoiser asset.

#tests Editor, Win64, ran and profiled NNE denoiser models

#rnx

diff --git a/Engine/Plugins/NNE/NNEDenoiser/Source/NNEDenoiser/Private/NNEDenoiserViewExtension.cpp b/Engine/Plugins/NNE/NNEDenoiser/Source/NNEDenoiser/Private/NNEDenoiserViewExtension.cpp
--- a/Engine/Plugins/NNE/NNEDenoiser/Source/NNEDenoiser/Private/NNEDenoiserViewExtension.cpp
+++ b/Engine/Plugins/NNE/NNEDenoiser/Source/NNEDenoiser/Private/NNEDenoiserViewExtension.cpp
@@ -260,7 +260,7 @@ FParameters GetParametersValidated(const AssetType& DenoiserAsset)
 }
 
 template<class AssetType, class InputMappingType, class OutputMappingType>
-TUniquePtr<FGenericDenoiser> CreateNNEDenoiserFromAsset(const FString& AssetName, EDenoiserRuntimeType RuntimeType, const FString& RuntimeNameOverride)
+TUniquePtr<FGenericDenoiser> CreateNNEDenoiserFromAsset(const FString& AssetName, EDenoiserRuntimeType RuntimeType, const FString& RuntimeNameOverride, int32 MaximumTileSizeOverride)
 {
 	if (AssetName.IsEmpty())
 	{
@@ -317,6 +317,25 @@ TUniquePtr<FGenericDenoiser> CreateNNEDenoiserFromAsset(const FString& AssetName
 
 	FParameters Parameters = GetParametersValidated(*DenoiserAsset);
 
+	if (MaximumTileSizeOverride > -1)
+	{
+		if (MaximumTileSizeOverride < Parameters.TilingConfig.MaxSize)
+		{
+			if (MaximumTileSizeOverride < Parameters.TilingConfig.MinSize)
+			{
+				UE_LOG(LogNNEDenoiser, Display, TEXT("Maximum tile size override %d smaller than minimum tile size %d supported by asset!"), MaximumTileSizeOverride, Parameters.TilingConfig.MinSize);
+			}
+
+			int32 NewMaxSize = FMath::Clamp(MaximumTileSizeOverride, Parameters.TilingConfig.MinSize, Parameters.TilingConfig.MaxSize);
+			if (NewMaxSize < Parameters.TilingConfig.MaxSize)
+			{
+				UE_LOG(LogNNEDenoiser, Log, TEXT("Applied maximum tile size override %d to assets maximum tile size %d."), NewMaxSize, Parameters.TilingConfig.MaxSize);
+
+				Parameters.TilingConfig.MaxSize = NewMaxSize;
+			}
+		}
+	}
+
 	return CreateNNEDenoiser(
 		*ModelData,
 		RuntimeType,
@@ -357,6 +376,7 @@ void FViewExtension::BeginRenderViewFamily(FSceneViewFamily& InViewFamily)
 	bNeedsUpdate |= CVarNNEDenoiserRuntimeName.GetValueOnGameThread() != RuntimeName;
 	bNeedsUpdate |= GetDenoiserAssetNameFromCVarAndSettings(Settings) != AssetName;
 	bNeedsUpdate |= GetDenoiserTemporalAssetNameFromCVarAndSettings(Settings) != TemporalAssetName;
+	bNeedsUpdate |= Settings->MaximumTileSizeOverride != MaximumTileSizeOverride;
 
 	if (!bNeedsUpdate)
 	{
@@ -375,10 +395,11 @@ void FViewExtension::BeginRenderViewFamily(FSceneViewFamily& InViewFamily)
 	RuntimeName = CVarNNEDenoiserRuntimeName.GetValueOnGameThread();
 	AssetName = GetDenoiserAssetNameFromCVarAndSettings(Settings);
 	TemporalAssetName = GetDenoiserTemporalAssetNameFromCVarAndSettings(Settings);
+	MaximumTileSizeOverride = Settings->MaximumTileSizeOverride;
 
 	// UE_LOG(LogNNEDenoiser, Log, TEXT("Create denoiser from asset %s..."), *AssetName);
 
-	TUniquePtr<FGenericDenoiser> Denoiser = CreateNNEDenoiserFromAsset<UNNEDenoiserAsset, FNNEDenoiserInputMappingData, FNNEDenoiserOutputMappingData>(AssetName, RuntimeType, RuntimeName);
+	TUniquePtr<FGenericDenoiser> Denoiser = CreateNNEDenoiserFromAsset<UNNEDenoiserAsset, FNNEDenoiserInputMappingData, FNNEDenoiserOutputMappingData>(AssetName, RuntimeType, RuntimeName, MaximumTileSizeOverride);
 	if (Denoiser.IsValid())
 	{
 		UE_LOG(LogNNEDenoiser, Log, TEXT("Create denoiser from asset %s..."), *AssetName);
@@ -390,7 +411,7 @@ void FViewExtension::BeginRenderViewFamily(FSceneViewFamily& InViewFamily)
 		UE_LOG(LogNNEDenoiser, Error, TEXT("Could not create denoiser from asset %d!"), *AssetName);
 	}
 
-	TUniquePtr<FGenericDenoiser> TemporalDenoiser = CreateNNEDenoiserFromAsset<UNNEDenoiserTemporalAsset, FNNEDenoiserTemporalInputMappingData, FNNEDenoiserTemporalOutputMappingData>(TemporalAssetName, RuntimeType, RuntimeName);
+	TUniquePtr<FGenericDenoiser> TemporalDenoiser = CreateNNEDenoiserFromAsset<UNNEDenoiserTemporalAsset, FNNEDenoiserTemporalInputMappingData, FNNEDenoiserTemporalOutputMappingData>(TemporalAssetName, RuntimeType, RuntimeName, MaximumTileSizeOverride);
 	if (TemporalDenoiser.IsValid())
 	{
 		UE_LOG(LogNNEDenoiser, Log, TEXT("Create temporal denoiser from asset %s..."), *TemporalAssetName);
diff --git a/Engine/Plugins/NNE/NNEDenoiser/Source/NNEDenoiser/Private/NNEDenoiserViewExtension.h b/Engine/Plugins/NNE/NNEDenoiser/Source/NNEDenoiser/Private/NNEDenoiserViewExtension.h
--- a/Engine/Plugins/NNE/NNEDenoiser/Source/NNEDenoiser/Private/NNEDenoiserViewExtension.h
+++ b/Engine/Plugins/NNE/NNEDenoiser/Source/NNEDenoiser/Private/NNEDenoiserViewExtension.h
@@ -35,6 +35,7 @@ class FViewExtension final : public FSceneViewExtensionBase
 		FString RuntimeName;
 		FString AssetName;
 		FString TemporalAssetName;
+		uint32 MaximumTileSizeOverride;
 	};
 
 } // namespace UE::NNEDenoiser::Private
\ No newline at end of file
diff --git a/Engine/Plugins/NNE/NNEDenoiser/Source/NNEDenoiser/Public/NNEDenoiserSettings.h b/Engine/Plugins/NNE/NNEDenoiser/Source/NNEDenoiser/Public/NNEDenoiserSettings.h
--- a/Engine/Plugins/NNE/NNEDenoiser/Source/NNEDenoiser/Public/NNEDenoiserSettings.h
+++ b/Engine/Plugins/NNE/NNEDenoiser/Source/NNEDenoiser/Public/NNEDenoiserSettings.h
@@ -35,6 +35,10 @@ public:
 	/** Temporal denoiser asset data used to create a NNE Denoiser (Currently not used and therefore "hidden") */
 	TSoftObjectPtr<UNNEDenoiserTemporalAsset> TemporalDenoiserAsset;
 
+	/** Override the maximum tile size defined per asset, but be aware not to set it lower than the assets minimum tile size (This can reduce GPU memory usage for GPU and RDG backed denoisers). */
+	UPROPERTY(Config, EditAnywhere, Category = "NNE Denoiser", meta = (DisplayName = "Maximum tile size override", ToolTip = "Override the maximum tile size given by the asset.\nSpecial values:\n   -1 = Do not override"))
+	int32 MaximumTileSizeOverride = -1;
+
 private:
 	/** Runtime type used to run the NNE Denoiser model. Backed by the console variable 'NNEDenoiser.Runtime.Type'. */
 	UPROPERTY(Config, EditAnywhere, Category = "NNE Denoiser", meta = (DisplayName = "Runtime Type", ToolTip = "Select a Runtime type", ConsoleVariable = "NNEDenoiser.Runtime.Type"))
