[NNEDenoiser] More carefully handle multiple viewports with no all of them using the Path Tracer. Avoid loading/unloading denoiser in every frame.

#tests Editor, Win64, tested with multiple view ports

#rnx

diff --git a/Engine/Plugins/NNE/NNEDenoiser/Source/NNEDenoiser/Private/NNEDenoiserViewExtension.cpp b/Engine/Plugins/NNE/NNEDenoiser/Source/NNEDenoiser/Private/NNEDenoiserViewExtension.cpp
--- a/Engine/Plugins/NNE/NNEDenoiser/Source/NNEDenoiser/Private/NNEDenoiserViewExtension.cpp
+++ b/Engine/Plugins/NNE/NNEDenoiser/Source/NNEDenoiser/Private/NNEDenoiserViewExtension.cpp
@@ -24,6 +24,8 @@
 namespace UE::NNEDenoiser::Private
 {
 
+constexpr int32 NO_SHOW_PATH_TRACING_TIMEOUT = 300; // Number of calls to BeginRenderViewFamily() wihtout Path Tracing show flag set in view family
+
 static TAutoConsoleVariable<bool> CVarNNEDenoiser(
 	TEXT("NNEDenoiser"),
 	true,
@@ -111,16 +113,6 @@ FString GetDenoiserTemporalAssetNameFromCVarAndSettings(const UNNEDenoiserSettin
 	return FString();
 }
 
-bool IsPathTracingEnabled(FSceneViewFamily& InViewFamily)
-{
-	if (!IsRayTracingEnabled())
-	{
-		return false;
-	}
-
-	return !!InViewFamily.EngineShowFlags.PathTracing;
-}
-
 TUniquePtr<FGenericDenoiser> CreateNNEDenoiser(
 	UNNEModelData& ModelData,
 	EDenoiserRuntimeType RuntimeType,
@@ -366,12 +358,30 @@ FViewExtension::~FViewExtension()
 
 void FViewExtension::BeginRenderViewFamily(FSceneViewFamily& InViewFamily)
 {
-	const UNNEDenoiserSettings* Settings = GetDefault<UNNEDenoiserSettings>();
+	TimeoutCounter = !!InViewFamily.EngineShowFlags.PathTracing ? NO_SHOW_PATH_TRACING_TIMEOUT : TimeoutCounter > 0 ? TimeoutCounter - 1 : 0;
+
+	const bool bIsActiveBefore = bIsActive;
+	bIsActive = IsRayTracingEnabled() && CVarNNEDenoiser.GetValueOnGameThread() && TimeoutCounter > 0;
+	if (bIsActive != bIsActiveBefore)
+	{
+		if (bIsActive)
+		{
+			UE_LOG(LogNNEDenoiser, Log, TEXT("Activating..."));
+		}
+		else
+		{
+			UE_LOG(LogNNEDenoiser, Log, TEXT("Deactivating and free all resources..."));
+		}
+	}
+
+	if (!bIsActive)
+	{
+		return;
+	}
 
 	// Check for changes
-	bool bNeedsUpdate = false;
-	bNeedsUpdate |= IsPathTracingEnabled(InViewFamily) != bPathTracingEnabled;
-	bNeedsUpdate |= CVarNNEDenoiser.GetValueOnGameThread() != bDenoiserEnabled;
+	const UNNEDenoiserSettings* Settings = GetDefault<UNNEDenoiserSettings>();
+	bool bNeedsUpdate = bIsActive != bIsActiveBefore;
 	bNeedsUpdate |= GetDenoiserRuntimeTypeFromCVar() != RuntimeType;
 	bNeedsUpdate |= CVarNNEDenoiserRuntimeName.GetValueOnGameThread() != RuntimeName;
 	bNeedsUpdate |= GetDenoiserAssetNameFromCVarAndSettings(Settings) != AssetName;
@@ -383,22 +393,12 @@ void FViewExtension::BeginRenderViewFamily(FSceneViewFamily& InViewFamily)
 		return;
 	}
 
-	bPathTracingEnabled = IsPathTracingEnabled(InViewFamily);
-	bDenoiserEnabled = CVarNNEDenoiser.GetValueOnGameThread();
-
-	if (!bPathTracingEnabled || !bDenoiserEnabled)
-	{
-		return;
-	}
-
 	RuntimeType = GetDenoiserRuntimeTypeFromCVar();
 	RuntimeName = CVarNNEDenoiserRuntimeName.GetValueOnGameThread();
 	AssetName = GetDenoiserAssetNameFromCVarAndSettings(Settings);
 	TemporalAssetName = GetDenoiserTemporalAssetNameFromCVarAndSettings(Settings);
 	MaximumTileSizeOverride = Settings->MaximumTileSizeOverride;
 
-	// UE_LOG(LogNNEDenoiser, Log, TEXT("Create denoiser from asset %s..."), *AssetName);
-
 	TUniquePtr<FGenericDenoiser> Denoiser = CreateNNEDenoiserFromAsset<UNNEDenoiserAsset, FNNEDenoiserInputMappingData, FNNEDenoiserOutputMappingData>(AssetName, RuntimeType, RuntimeName, MaximumTileSizeOverride);
 	if (Denoiser.IsValid())
 	{
@@ -426,7 +426,7 @@ void FViewExtension::BeginRenderViewFamily(FSceneViewFamily& InViewFamily)
 
 void FViewExtension::PreRenderViewFamily_RenderThread(FRDGBuilder& GraphBuilder, FSceneViewFamily& InViewFamily)
 {
-	if (!bPathTracingEnabled || !bDenoiserEnabled)
+	if (!bIsActive)
 	{
 		UnregisterDenoiser(TEXT("NNEDenoiser"));
 	}
diff --git a/Engine/Plugins/NNE/NNEDenoiser/Source/NNEDenoiser/Private/NNEDenoiserViewExtension.h b/Engine/Plugins/NNE/NNEDenoiser/Source/NNEDenoiser/Private/NNEDenoiserViewExtension.h
--- a/Engine/Plugins/NNE/NNEDenoiser/Source/NNEDenoiser/Private/NNEDenoiserViewExtension.h
+++ b/Engine/Plugins/NNE/NNEDenoiser/Source/NNEDenoiser/Private/NNEDenoiserViewExtension.h
@@ -29,8 +29,8 @@ class FViewExtension final : public FSceneViewExtensionBase
 		TUniquePtr<UE::Renderer::Private::IPathTracingSpatialTemporalDenoiser> SpatialTemporalDenoiserToSwap;
 
 		// Cached settings and CVar values
-		bool bPathTracingEnabled = false;
-		bool bDenoiserEnabled = true;
+		bool bIsActive = false;
+		int32 TimeoutCounter = 0;
 		EDenoiserRuntimeType RuntimeType;
 		FString RuntimeName;
 		FString AssetName;
