[NNE] NNEDenoiser added dump buffers.

#rnx
#Synced-CL 49101488

diff --git a/Engine/Plugins/NNE/NNEDenoiser/Source/NNEDenoiser/Private/NNEDenoiserDump.cpp b/Engine/Plugins/NNE/NNEDenoiser/Source/NNEDenoiser/Private/NNEDenoiserDump.cpp
new file mode 100644
--- /dev/null
+++ b/Engine/Plugins/NNE/NNEDenoiser/Source/NNEDenoiser/Private/NNEDenoiserDump.cpp
@@ -0,0 +1,126 @@
+// Copyright Epic Games, Inc. All Rights Reserved.
+
+#include "NNEDenoiserDump.h"
+
+#ifdef WITH_NNEDENOISER_DUMP
+#include "HAL/PlatformFileManager.h"
+#include "HAL/PreprocessorHelpers.h"
+#include "HAL/UnrealMemory.h"
+#include "Logging/LogVerbosity.h"
+#include "Misc/FileHelper.h"
+#include "Misc/Paths.h"
+#include "NNEDenoiserLog.h"
+#include "RenderGraphBuilder.h"
+#include "RenderGraphDefinitions.h"
+#include "RenderGraphEvent.h"
+#include "RenderingThread.h"
+#include "RHIGPUReadback.h"
+#include "Tasks/Task.h"
+#include "Windows/WindowsPlatformProcess.h"
+
+BEGIN_SHADER_PARAMETER_STRUCT(FDumpBufferPassParameters, )
+	RDG_BUFFER_ACCESS(Buffer, ERHIAccess::CopySrc)
+END_SHADER_PARAMETER_STRUCT()
+
+namespace UE::NNEDenoiser::Private
+{
+
+FString GetDefaultDumpPath()
+{
+	const FDateTime NowUTC = FDateTime::UtcNow();
+	const FString Stamp = NowUTC.ToString(TEXT("%Y-%m-%d_%H-%M-%S"));
+	FString Base = FString::Printf(TEXT("NNEDenoiserDump_%s"), *Stamp);
+	const FString DumpFilePath = FPaths::ProjectSavedDir() / TEXT("NNEDenoiserDumps") / FPaths::MakeValidFileName(Base);
+	
+	return DumpFilePath;
+}
+
+void ScheduleDumpBuffer(FRDGBuilder& GraphBuilder, FRDGBufferRef Buffer, const FString& Name, const FString& FilePath)
+{
+	FString FilePathLocal = FilePath;
+
+	static FString DefaultPath;
+	if (FilePathLocal.IsEmpty())
+	{
+		if (DefaultPath.IsEmpty())
+		{
+			DefaultPath = GetDefaultDumpPath();
+		}
+		FilePathLocal = DefaultPath;
+	}
+	else
+	{
+		DefaultPath = FilePathLocal;
+	}
+
+	FPlatformFileManager::Get().GetPlatformFile().CreateDirectoryTree(*FilePathLocal);
+
+	TSharedRef<FRHIGPUBufferReadback> Readback = MakeShared<FRHIGPUBufferReadback>(FName(Name));
+	FRHIGPUBufferReadback* ReadbackPtr = &Readback.Get();
+
+	auto* PassParameters = GraphBuilder.AllocParameters<FDumpBufferPassParameters>();
+	PassParameters->Buffer = Buffer;
+
+	GraphBuilder.AddPass(
+		RDG_EVENT_NAME("DebugDumpBuffer(%s)", *Name),
+		PassParameters,
+		ERDGPassFlags::Copy | ERDGPassFlags::NeverCull,
+		[PassParameters, ReadbackPtr](FRHICommandList& RHICommandList)
+		{
+			ReadbackPtr->EnqueueCopy(RHICommandList, PassParameters->Buffer->GetRHI());
+		}
+	);
+
+	Tasks::Launch(UE_SOURCE_LOCATION, [Readback, Name, FilePathLocal]()
+	{
+		FEvent* Signal = FGenericPlatformProcess::GetSynchEventFromPool();
+		TArray<uint8> FileData;
+
+		while (!Readback->IsReady())
+		{
+			FPlatformProcess::Sleep(0.01f);
+		}
+
+		ENQUEUE_RENDER_COMMAND(ReadRHIBuffer)([Readback, Name, Signal, &FileData](FRHICommandListImmediate& RHICmdList)
+		{	
+			int32 Size = Readback->GetGPUSizeBytes();
+			if (Size > 0)
+			{
+				void* Data = Readback->Lock(Size);
+				if (Data)
+				{
+					FileData.SetNumUninitialized(Size);
+					FMemory::Memcpy(FileData.GetData(), Data, Size);
+				}
+				Readback->Unlock();
+			}
+			
+			Signal->Trigger();
+		});
+
+		Signal->Wait();
+
+		FGenericPlatformProcess::ReturnSynchEventToPool(Signal);
+		Signal = nullptr;
+
+		if (FileData.IsEmpty())
+		{
+			UE_LOG(LogNNEDenoiser, Warning, TEXT("No data written for: %s (empty readback)"), *Name);
+			return;
+		}
+
+		FString Filename = (!FilePathLocal.IsEmpty() ? FilePathLocal : FPaths::ProjectSavedDir()) / Name + TEXT(".bin");
+		if (FFileHelper::SaveArrayToFile(FileData, *Filename))
+		{
+			UE_LOG(LogNNEDenoiser, Log, TEXT("Dumped buffer '%s' to %s (%d bytes)"), *Name, *Filename, FileData.Num());
+		}
+		else
+		{
+			UE_LOG(LogNNEDenoiser, Error, TEXT("Failed to write dump for '%s' to %s"), *Name, *Filename);
+		}
+	});
+}
+
+} // namespace UE::NNEDenoiser::Private
+
+#endif // WITH_NNEDENOISER_DUMP
\ No newline at end of file
diff --git a/Engine/Plugins/NNE/NNEDenoiser/Source/NNEDenoiser/Private/NNEDenoiserDump.h b/Engine/Plugins/NNE/NNEDenoiser/Source/NNEDenoiser/Private/NNEDenoiserDump.h
new file mode 100644
--- /dev/null
+++ b/Engine/Plugins/NNE/NNEDenoiser/Source/NNEDenoiser/Private/NNEDenoiserDump.h
@@ -0,0 +1,26 @@
+// Copyright Epic Games, Inc. All Rights Reserved.
+
+#pragma once
+
+#ifdef WITH_NNEDENOISER_DUMP
+
+#include "RenderGraphFwd.h"
+#include "Containers/UnrealString.h"
+
+namespace UE::NNEDenoiser::Private
+{
+
+FString GetDefaultDumpPath();
+
+void ScheduleDumpBuffer(FRDGBuilder& GraphBuilder, FRDGBufferRef Buffer, const FString& Name, const FString& FilePath = FString());
+
+} // namespace UE::NNEDenoiser::Private
+
+#define NNEDENOISER_DUMP_BUFFERS_SCHEDULE(GraphBuilder, Buffers, Prefix) \
+	for (int32 i = 0; i < Buffers.Num(); i++) \
+	{ \
+		UE::NNEDenoiser::Private::ScheduleDumpBuffer(GraphBuilder, Buffers[i], FString::Printf(TEXT("%s_%d"), Prefix, i)); \
+	}
+#else
+#define NNEDENOISER_DUMP_BUFFERS_SCHEDULE(GraphBuilder, Buffers, Prefix)
+#endif // WITH_NNEDENOISER_DUMP
\ No newline at end of file
diff --git a/Engine/Plugins/NNE/NNEDenoiser/Source/NNEDenoiser/Private/NNEDenoiserGenericDenoiser.cpp b/Engine/Plugins/NNE/NNEDenoiser/Source/NNEDenoiser/Private/NNEDenoiserGenericDenoiser.cpp
--- a/Engine/Plugins/NNE/NNEDenoiser/Source/NNEDenoiser/Private/NNEDenoiserGenericDenoiser.cpp
+++ b/Engine/Plugins/NNE/NNEDenoiser/Source/NNEDenoiser/Private/NNEDenoiserGenericDenoiser.cpp
@@ -3,6 +3,7 @@
 #include "NNEDenoiserGenericDenoiser.h"
 #include "Algo/Transform.h"
 #include "NNEDenoiserAutoExposure.h"
+#include "NNEDenoiserDump.h"
 #include "NNEDenoiserHistory.h"
 #include "NNEDenoiserIOProcess.h"
 #include "NNEDenoiserLog.h"
@@ -126,10 +127,14 @@ void AddTilePasses(FRDGBuilder& GraphBuilder, IModelInstance& ModelInstance, con
 	// 1. Read input and write to input buffers
 	InputProcess.AddPasses(GraphBuilder, ModelInstance.GetInputTensorDescs(), ModelInstance.GetInputTensorShapes(), ResourceAccess, InputBuffers);
 
+	NNEDENOISER_DUMP_BUFFERS_SCHEDULE(GraphBuilder, InputBuffers, TEXT("NNE_input"))
+	
 	// 2. Create buffer binding and infer the model
 	NNE::IModelInstanceRDG::EEnqueueRDGStatus Status = ModelInstance.EnqueueRDG(GraphBuilder, GetBindingRDG(InputBuffers), GetBindingRDG(OutputBuffers));
 	checkf(Status == NNE::IModelInstanceRDG::EEnqueueRDGStatus::Ok, TEXT("EnqueueRDG failed: %d"), static_cast<int>(Status));
 
+	NNEDENOISER_DUMP_BUFFERS_SCHEDULE(GraphBuilder, OutputBuffers, TEXT("NNE_output"))
+
 	// 3. Write output based on output buffer
 	FRDGTextureRef OutputTexture = ResourceManager.GetTexture(EResourceName::Output, 0);
 	OutputProcess.AddPasses(GraphBuilder, ModelInstance.GetOutputTensorDescs(), GetOutputTensorShapes(ModelInstance), ResourceAccess, OutputBuffers, OutputTexture);
