[NNEDenoiser] Fixed buffer binding size for CPU and GPU runtime interface.

#tests Editor, Mac, ran PT

#rnx

#Synced-CL 50857221

diff --git a/Engine/Plugins/NNE/NNEDenoiser/Source/NNEDenoiser/Private/NNEDenoiserModelInstanceCPU.cpp b/Engine/Plugins/NNE/NNEDenoiser/Source/NNEDenoiser/Private/NNEDenoiserModelInstanceCPU.cpp
--- a/Engine/Plugins/NNE/NNEDenoiser/Source/NNEDenoiser/Private/NNEDenoiserModelInstanceCPU.cpp
+++ b/Engine/Plugins/NNE/NNEDenoiser/Source/NNEDenoiser/Private/NNEDenoiserModelInstanceCPU.cpp
@@ -123,7 +123,7 @@ namespace UE::NNEDenoiser::Private
 			TArray<NNE::FTensorBindingCPU> InputBindings;
 			for (int32 Idx = 0; Idx < DenoiserParameters->InputBuffers.Num(); Idx++)
 			{
-				ScratchInputBuffers[Idx].SetNumUninitialized(DenoiserParameters->InputBuffers[Idx].GetBuffer()->GetRHI()->GetSize());
+				ScratchInputBuffers[Idx].SetNumUninitialized(DenoiserParameters->InputBuffers[Idx].GetBuffer()->GetSize());
 
 				InputBindings.Emplace(NNE::FTensorBindingCPU{(void *)ScratchInputBuffers[Idx].GetData(), (uint64)ScratchInputBuffers[Idx].Num()});
 			}
@@ -132,7 +132,7 @@ namespace UE::NNEDenoiser::Private
 			TArray<NNE::FTensorBindingCPU> OutputBindings;
 			for (int32 Idx = 0; Idx < DenoiserParameters->OutputBuffers.Num(); Idx++)
 			{
-				ScratchOutputBuffers[Idx].SetNumUninitialized(DenoiserParameters->OutputBuffers[Idx].GetBuffer()->GetRHI()->GetSize());
+				ScratchOutputBuffers[Idx].SetNumUninitialized(DenoiserParameters->OutputBuffers[Idx].GetBuffer()->GetSize());
 
 				OutputBindings.Emplace(NNE::FTensorBindingCPU{(void *)ScratchOutputBuffers[Idx].GetData(), (uint64)ScratchOutputBuffers[Idx].Num()});
 			}
@@ -140,7 +140,8 @@ namespace UE::NNEDenoiser::Private
 			for (int32 Idx = 0; Idx < DenoiserParameters->InputBuffers.Num(); Idx++)
 			{
 				FRHIBuffer* Buffer = DenoiserParameters->InputBuffers[Idx].GetBuffer()->GetRHI();
-				CopyBufferFromGPUToCPU(RHICmdList, Buffer, Buffer->GetSize(), ScratchInputBuffers[Idx]);
+				const int32 Count = static_cast<int32>(DenoiserParameters->InputBuffers[Idx]->GetSize());
+				CopyBufferFromGPUToCPU(RHICmdList, Buffer, Count, ScratchInputBuffers[Idx]);
 			}
 
 			NNE::IModelInstanceCPU::ERunSyncStatus Status = ModelInstance->RunSync(InputBindings, OutputBindings);
@@ -149,7 +150,8 @@ namespace UE::NNEDenoiser::Private
 			for (int32 Idx = 0; Idx < DenoiserParameters->OutputBuffers.Num(); Idx++)
 			{
 				FRHIBuffer* Buffer = DenoiserParameters->OutputBuffers[Idx].GetBuffer()->GetRHI();
-				CopyBufferFromCPUToGPU(RHICmdList, ScratchOutputBuffers[Idx], Buffer->GetSize(), Buffer);
+				const int32 Count = static_cast<int32>(DenoiserParameters->OutputBuffers[Idx]->GetSize());
+				CopyBufferFromCPUToGPU(RHICmdList, ScratchOutputBuffers[Idx], Count, Buffer);
 			}
 
 #if WITH_EDITOR
diff --git a/Engine/Plugins/NNE/NNEDenoiser/Source/NNEDenoiser/Private/NNEDenoiserModelInstanceGPU.cpp b/Engine/Plugins/NNE/NNEDenoiser/Source/NNEDenoiser/Private/NNEDenoiserModelInstanceGPU.cpp
--- a/Engine/Plugins/NNE/NNEDenoiser/Source/NNEDenoiser/Private/NNEDenoiserModelInstanceGPU.cpp
+++ b/Engine/Plugins/NNE/NNEDenoiser/Source/NNEDenoiser/Private/NNEDenoiserModelInstanceGPU.cpp
@@ -128,7 +128,7 @@ namespace UE::NNEDenoiser::Private
 			TArray<NNE::FTensorBindingCPU> InputBindings;
 			for (int32 Idx = 0; Idx < DenoiserParameters->InputBuffers.Num(); Idx++)
 			{
-				ScratchInputBuffers[Idx].SetNumUninitialized(DenoiserParameters->InputBuffers[Idx].GetBuffer()->GetRHI()->GetSize());
+				ScratchInputBuffers[Idx].SetNumUninitialized(DenoiserParameters->InputBuffers[Idx].GetBuffer()->GetSize());
 
 				InputBindings.Emplace(NNE::FTensorBindingCPU{(void *)ScratchInputBuffers[Idx].GetData(), (uint64)ScratchInputBuffers[Idx].Num()});
 			}
@@ -137,7 +137,7 @@ namespace UE::NNEDenoiser::Private
 			TArray<NNE::FTensorBindingCPU> OutputBindings;
 			for (int32 Idx = 0; Idx < DenoiserParameters->OutputBuffers.Num(); Idx++)
 			{
-				ScratchOutputBuffers[Idx].SetNumUninitialized(DenoiserParameters->OutputBuffers[Idx].GetBuffer()->GetRHI()->GetSize());
+				ScratchOutputBuffers[Idx].SetNumUninitialized(DenoiserParameters->OutputBuffers[Idx].GetBuffer()->GetSize());
 
 				OutputBindings.Emplace(NNE::FTensorBindingCPU{(void *)ScratchOutputBuffers[Idx].GetData(), (uint64)ScratchOutputBuffers[Idx].Num()});
 			}
@@ -145,7 +145,8 @@ namespace UE::NNEDenoiser::Private
 			for (int32 Idx = 0; Idx < DenoiserParameters->InputBuffers.Num(); Idx++)
 			{
 				FRHIBuffer* Buffer = DenoiserParameters->InputBuffers[Idx].GetBuffer()->GetRHI();
-				CopyBufferFromGPUToCPU(RHICmdList, Buffer, Buffer->GetSize(), ScratchInputBuffers[Idx]);
+				const int32 Count = static_cast<int32>(DenoiserParameters->InputBuffers[Idx]->GetSize());
+				CopyBufferFromGPUToCPU(RHICmdList, Buffer, Count, ScratchInputBuffers[Idx]);
 			}
 
 			NNE::IModelInstanceGPU::ERunSyncStatus Status = ModelInstance->RunSync(InputBindings, OutputBindings);
@@ -154,7 +155,8 @@ namespace UE::NNEDenoiser::Private
 			for (int32 Idx = 0; Idx < DenoiserParameters->OutputBuffers.Num(); Idx++)
 			{
 				FRHIBuffer* Buffer = DenoiserParameters->OutputBuffers[Idx].GetBuffer()->GetRHI();
-				CopyBufferFromCPUToGPU(RHICmdList, ScratchOutputBuffers[Idx], Buffer->GetSize(), Buffer);
+				const int32 Count = static_cast<int32>(DenoiserParameters->OutputBuffers[Idx]->GetSize());
+				CopyBufferFromCPUToGPU(RHICmdList, ScratchOutputBuffers[Idx], Count, Buffer);
 			}
 
 #if WITH_EDITOR
